Rename SkillSystemEditor module to GameplayAbilitiesEditor

diff --git a/Engine/Source/Editor/GameplayAbilitiesEditor/Classes/K2Node_PlayMontageAndWait.h b/Engine/Source/Editor/GameplayAbilitiesEditor/Classes/K2Node_PlayMontageAndWait.h
new file mode 100644
--- /dev/null
+++ b/Engine/Source/Editor/GameplayAbilitiesEditor/Classes/K2Node_PlayMontageAndWait.h
@@ -0,0 +1,17 @@
+// Copyright 1998-2014 Epic Games, Inc. All Rights Reserved.
+
+#pragma once
+#include "K2Node_PlayMontageAndWait.generated.h"
+
+UCLASS()
+class UK2Node_PlayMontageAndWait : public UK2Node_BaseAsyncTask
+{
+	GENERATED_UCLASS_BODY()
+
+	// Begin UEdGraphNode interface
+	virtual FString GetTooltip() const OVERRIDE;
+	virtual FText GetNodeTitle(ENodeTitleType::Type TitleType) const OVERRIDE;
+	// End UEdGraphNode interface
+
+	virtual FString GetCategoryName();
+};
diff --git a/Engine/Source/Editor/GameplayAbilitiesEditor/Classes/K2Node_WaitMovementModeChange.h b/Engine/Source/Editor/GameplayAbilitiesEditor/Classes/K2Node_WaitMovementModeChange.h
new file mode 100644
--- /dev/null
+++ b/Engine/Source/Editor/GameplayAbilitiesEditor/Classes/K2Node_WaitMovementModeChange.h
@@ -0,0 +1,17 @@
+// Copyright 1998-2014 Epic Games, Inc. All Rights Reserved.
+
+#pragma once
+#include "K2Node_WaitMovementModeChange.generated.h"
+
+UCLASS()
+class UK2Node_WaitMovementModeChange : public UK2Node_BaseAsyncTask
+{
+	GENERATED_UCLASS_BODY()
+
+	// Begin UEdGraphNode interface
+	virtual FString GetTooltip() const OVERRIDE;
+	virtual FText GetNodeTitle(ENodeTitleType::Type TitleType) const OVERRIDE;
+	// End UEdGraphNode interface
+
+	virtual FString GetCategoryName();
+};
diff --git a/Engine/Source/Editor/GameplayAbilitiesEditor/Classes/K2Node_WaitOverlap.h b/Engine/Source/Editor/GameplayAbilitiesEditor/Classes/K2Node_WaitOverlap.h
new file mode 100644
--- /dev/null
+++ b/Engine/Source/Editor/GameplayAbilitiesEditor/Classes/K2Node_WaitOverlap.h
@@ -0,0 +1,17 @@
+// Copyright 1998-2014 Epic Games, Inc. All Rights Reserved.
+
+#pragma once
+#include "K2Node_WaitOverlap.generated.h"
+
+UCLASS()
+class UK2Node_WaitOverlap : public UK2Node_BaseAsyncTask
+{
+	GENERATED_UCLASS_BODY()
+
+	// Begin UEdGraphNode interface
+	virtual FString GetTooltip() const OVERRIDE;
+	virtual FText GetNodeTitle(ENodeTitleType::Type TitleType) const OVERRIDE;
+	// End UEdGraphNode interface
+
+	virtual FString GetCategoryName();
+};
diff --git a/Engine/Source/Editor/GameplayAbilitiesEditor/GameplayAbilitiesEditor.Build.cs b/Engine/Source/Editor/GameplayAbilitiesEditor/GameplayAbilitiesEditor.Build.cs
new file mode 100644
--- /dev/null
+++ b/Engine/Source/Editor/GameplayAbilitiesEditor/GameplayAbilitiesEditor.Build.cs
@@ -0,0 +1,37 @@
+// Copyright 1998-2013 Epic Games, Inc. All Rights Reserved.
+
+namespace UnrealBuildTool.Rules
+{
+	public class GameplayAbilitiesEditor : ModuleRules
+	{
+        public GameplayAbilitiesEditor(TargetInfo Target)
+		{
+            PrivateIncludePaths.AddRange(
+                new string[] {
+					"Editor/GameplayAbilitiesEditor/Private",
+				}
+			);
+
+            PrivateDependencyModuleNames.AddRange(
+                new string[]
+				{
+					// ... add private dependencies that you statically link with here ...
+					"Core",
+					"CoreUObject",
+					"Engine",
+					"AssetTools",
+					"GameplayAbilities",
+                    "InputCore",
+					"Slate",
+					"SlateCore",
+                    "EditorStyle",
+					"BlueprintGraph",
+					"KismetCompiler",
+					"GraphEditor",
+					"MainFrame",
+					"UnrealEd",
+				}
+			);
+		}
+	}
+}
\ No newline at end of file
diff --git a/Engine/Source/Editor/GameplayAbilitiesEditor/Private/AttributeDetails.cpp b/Engine/Source/Editor/GameplayAbilitiesEditor/Private/AttributeDetails.cpp
new file mode 100644
--- /dev/null
+++ b/Engine/Source/Editor/GameplayAbilitiesEditor/Private/AttributeDetails.cpp
@@ -0,0 +1,482 @@
+// Copyright 1998-2013 Epic Games, Inc. All Rights Reserved.
+
+#include "SkillSystemEditorModulePrivatePCH.h"
+#include "AttributeDetails.h"
+#include "Editor/PropertyEditor/Public/PropertyEditing.h"
+#include "AttributeDetails.h"
+#include "AttributeSet.h"
+#include "BlueprintUtilities.h"
+#include "KismetEditorUtilities.h"
+#include "SkillSystemGlobals.h"
+#include "GameplayAbilitiesModule.h"
+
+#define LOCTEXT_NAMESPACE "AttributeDetailsCustomization"
+
+// default name for base pose
+#define DEFAULT_POSE_NAME	TEXT("Default")
+
+DEFINE_LOG_CATEGORY(LogAttributeDetails);
+
+TSharedRef<IPropertyTypeCustomization> FAttributePropertyDetails::MakeInstance()
+{
+	return MakeShareable(new FAttributePropertyDetails());
+}
+
+BEGIN_SLATE_FUNCTION_BUILD_OPTIMIZATION
+void FAttributePropertyDetails::CustomizeHeader( TSharedRef<IPropertyHandle> StructPropertyHandle, class FDetailWidgetRow& HeaderRow, IStructCustomizationUtils& StructCustomizationUtils )
+{
+	MyProperty = StructPropertyHandle->GetChildHandle(GET_MEMBER_NAME_CHECKED(FGameplayAttribute,Attribute));
+
+	PropertyOptions.Empty();
+	PropertyOptions.Add(MakeShareable(new FString("None")));
+
+	FString FilterMetaStr = StructPropertyHandle->GetProperty()->GetMetaData(TEXT("FilterMetaTag"));
+
+	// Gather all UAttraibute classes
+	for (TObjectIterator<UClass> ClassIt; ClassIt; ++ClassIt)
+	{
+		UClass *Class = *ClassIt;
+		if (Class->IsChildOf(UAttributeSet::StaticClass()) && !FKismetEditorUtilities::IsClassABlueprintSkeleton(Class))
+		{
+			for (TFieldIterator<UProperty> PropertyIt(Class, EFieldIteratorFlags::ExcludeSuper); PropertyIt; ++PropertyIt)
+			{
+				UProperty *Property = *PropertyIt;
+
+				if (!FilterMetaStr.IsEmpty() && Property->HasMetaData(*FilterMetaStr))
+				{
+					continue;
+				}
+				
+				PropertyOptions.Add(MakeShareable(new FString(FString::Printf(TEXT("%s.%s"), *Class->GetName(), *Property->GetName()))));
+			}
+		}
+	}
+
+	HeaderRow.
+		NameContent()
+		[
+			StructPropertyHandle->CreatePropertyNameWidget()
+		]
+		.ValueContent()
+		.MinDesiredWidth(500)
+		.MaxDesiredWidth(4096)
+		[
+			SNew(SHorizontalBox)
+			+ SHorizontalBox::Slot()
+			//.FillWidth(1.0f)
+			.HAlign(HAlign_Fill)
+			.Padding(0.f, 0.f, 2.f, 0.f)
+			[
+
+				SNew(STextComboBox)
+				.ContentPadding(FMargin(2.0f, 2.0f))
+				.OptionsSource( &PropertyOptions )
+				.InitiallySelectedItem(GetPropertyType())
+				.OnSelectionChanged( this, &FAttributePropertyDetails::OnChangeProperty )
+			]
+		];
+}
+END_SLATE_FUNCTION_BUILD_OPTIMIZATION
+
+void FAttributePropertyDetails::CustomizeChildren( TSharedRef<IPropertyHandle> StructPropertyHandle, class IDetailChildrenBuilder& StructBuilder, IStructCustomizationUtils& StructCustomizationUtils )
+{
+
+
+}
+
+TSharedPtr<FString> FAttributePropertyDetails::GetPropertyType() const
+{
+	if (MyProperty.IsValid())
+	{
+		UProperty *PropertyValue = NULL;
+		UObject *ObjPtr = NULL;
+		MyProperty->GetValue(ObjPtr);
+		PropertyValue = Cast<UProperty>(ObjPtr);
+		if (PropertyValue)
+		{
+			FString FullString = PropertyValue->GetOuter()->GetName() + TEXT(".") + PropertyValue->GetName();
+			for (int32 i=0; i < PropertyOptions.Num(); ++i)
+			{
+				if (PropertyOptions[i].IsValid() && PropertyOptions[i].Get()->Equals(FullString))
+				{
+					return PropertyOptions[i];
+				}
+			}
+		}
+	}
+
+	return PropertyOptions[0]; // This should always be none
+}
+
+void FAttributePropertyDetails::OnChangeProperty(TSharedPtr<FString> ItemSelected, ESelectInfo::Type SelectInfo)
+{
+	if (ItemSelected.IsValid() && MyProperty.IsValid())
+	{		
+		FString FullString = *ItemSelected.Get();
+		FString ClassName;
+		FString PropertyName;
+
+		FullString.Split( TEXT("."), &ClassName, &PropertyName);
+
+		UClass *FoundClass = FindObject<UClass>(ANY_PACKAGE, *ClassName);
+		if (FoundClass)
+		{
+			UProperty *Property = FindField<UProperty>(FoundClass, *PropertyName);
+			if (Property)
+			{
+				const UObject *ObjPtr = Property;
+				MyProperty->SetValue(ObjPtr);
+				return;
+			}
+		}
+	}
+}
+
+// ------------------------------------------------------------------------------------
+// ------------------------------------------------------------------------------------
+// ------------------------------------------------------------------------------------
+// ------------------------------------------------------------------------------------
+// ------------------------------------------------------------------------------------
+
+TSharedRef<IDetailCustomization> FAttributeDetails::MakeInstance()
+{
+	return MakeShareable(new FAttributeDetails);
+}
+
+BEGIN_SLATE_FUNCTION_BUILD_OPTIMIZATION
+void FAttributeDetails::CustomizeDetails(IDetailLayoutBuilder& DetailLayout)
+{
+	MyProperty = DetailLayout.GetProperty("PropertyReference", UAttributeSet::StaticClass());
+
+	PropertyOptions.Empty();
+	PropertyOptions.Add(MakeShareable(new FString("None")));
+
+	for( TFieldIterator<UProperty> It(UAttributeSet::StaticClass(), EFieldIteratorFlags::ExcludeSuper) ; It ; ++It )
+	{
+		UProperty *Property = *It;
+		PropertyOptions.Add(MakeShareable(new FString(Property->GetName())));
+	}
+
+	IDetailCategoryBuilder& Category = DetailLayout.EditCategory("Variable", LOCTEXT("VariableDetailsCategory", "Variable").ToString());
+	const FSlateFontInfo DetailFontInfo = IDetailLayoutBuilder::GetDetailFont();
+	
+	Category.AddCustomRow( TEXT("Replication") )
+		//.Visibility(TAttribute<EVisibility>(this, &FBlueprintVarActionDetails::ReplicationVisibility))
+		.NameContent()
+		[
+			SNew(STextBlock)
+			.ToolTipText(LOCTEXT("PropertyType_Tooltip", "Which Property To Modify?"))
+			.Text( LOCTEXT("PropertyModifierInfo", "Property") )
+			.Font( DetailFontInfo )
+		]
+		.ValueContent()
+		[
+			SNew(STextComboBox)
+			.OptionsSource( &PropertyOptions )
+			.InitiallySelectedItem(GetPropertyType())
+			.OnSelectionChanged( this, &FAttributeDetails::OnChangeProperty )
+		];
+	
+}
+END_SLATE_FUNCTION_BUILD_OPTIMIZATION
+
+
+TSharedPtr<FString> FAttributeDetails::GetPropertyType() const
+{
+	if (!MyProperty.IsValid())
+		return PropertyOptions[0];
+
+	UProperty *PropertyValue = NULL;
+	UObject *ObjPtr = NULL;
+	MyProperty->GetValue(ObjPtr);
+	PropertyValue = Cast<UProperty>(ObjPtr);
+
+	if (PropertyValue != NULL)
+	{
+		for (int32 i=0; i < PropertyOptions.Num(); ++i)
+		{
+			if (PropertyOptions[i].IsValid() && PropertyOptions[i].Get()->Equals(PropertyValue->GetName()))
+			{
+				return PropertyOptions[i];
+			}
+		}
+	}
+
+	return PropertyOptions[0]; // This should always be none
+}
+
+void FAttributeDetails::OnChangeProperty(TSharedPtr<FString> ItemSelected, ESelectInfo::Type SelectInfo)
+{
+	if (!ItemSelected.IsValid())
+	{
+		return;
+	}
+
+	FString PropertyName = *ItemSelected.Get();
+
+	for( TFieldIterator<UProperty> It(UAttributeSet::StaticClass(), EFieldIteratorFlags::ExcludeSuper) ; It ; ++It )
+	{
+		const UProperty *Property = *It;
+		if (PropertyName == Property->GetName())
+		{
+			const UObject *ObjPtr = Property;
+			MyProperty->SetValue(ObjPtr);
+			return;
+		}
+	}
+}
+
+// ------------------------------------------------------------------------------------
+// ------------------------------------------------------------------------------------
+// ------------------------------------------------------------------------------------
+// ------------------------------------------------------------------------------------
+// ------------------------------------------------------------------------------------
+
+TSharedRef<IPropertyTypeCustomization> FScalableFloatDetails::MakeInstance()
+{
+	return MakeShareable(new FScalableFloatDetails());
+}
+
+BEGIN_SLATE_FUNCTION_BUILD_OPTIMIZATION
+void FScalableFloatDetails::CustomizeHeader( TSharedRef<class IPropertyHandle> StructPropertyHandle, class FDetailWidgetRow& HeaderRow, IStructCustomizationUtils& StructCustomizationUtils )
+{
+	uint32 NumChildren = 0;
+	StructPropertyHandle->GetNumChildren(NumChildren);
+
+	for (uint32 i = 0; i < NumChildren; i++)
+	{
+		FString PropName = StructPropertyHandle->GetChildHandle(i)->GetPropertyDisplayName();
+		PropName;
+	}
+
+	ValueProperty = StructPropertyHandle->GetChildHandle(GET_MEMBER_NAME_CHECKED(FScalableFloat,Value));
+	CurveTableHandleProperty = StructPropertyHandle->GetChildHandle(GET_MEMBER_NAME_CHECKED(FScalableFloat,Curve));
+
+	RowNameProperty = CurveTableHandleProperty->GetChildHandle(GET_MEMBER_NAME_CHECKED(FCurveTableRowHandle, RowName));
+	CurveTableProperty = CurveTableHandleProperty->GetChildHandle(GET_MEMBER_NAME_CHECKED(FCurveTableRowHandle, CurveTable));
+
+	CurrentSelectedItem = InitWidgetContent();
+
+	FSimpleDelegate OnCurveTableChangedDelegate = FSimpleDelegate::CreateSP(this, &FScalableFloatDetails::OnCurveTableChanged);
+	CurveTableProperty->SetOnPropertyValueChanged(OnCurveTableChangedDelegate);
+
+	HeaderRow
+		.NameContent()
+		[
+			StructPropertyHandle->CreatePropertyNameWidget()
+		]
+		.ValueContent()
+		.MinDesiredWidth( 500 )
+		.MaxDesiredWidth( 4096 )
+		[
+			SNew(SHorizontalBox)
+			.IsEnabled(TAttribute<bool>::Create(TAttribute<bool>::FGetter::CreateSP(this, &FScalableFloatDetails::IsEditable)))
+			+SHorizontalBox::Slot()
+			//.FillWidth(1.0f)
+			.HAlign(HAlign_Fill)
+			.Padding(0.f, 0.f, 2.f, 0.f)
+			[
+				ValueProperty->CreatePropertyValueWidget()
+			]
+		
+			+SHorizontalBox::Slot()
+			//.FillWidth(1.f)
+			.HAlign(HAlign_Fill)
+			.Padding(2.f, 2.f, 2.f, 2.f)
+			[
+				SAssignNew(RowNameComboButton, SComboButton)
+				.OnGetMenuContent(this, &FScalableFloatDetails::GetListContent)
+				.ContentPadding(FMargin(2.0f, 2.0f))
+				.ButtonContent()
+				[
+					SNew(STextBlock)
+					.Text(this, &FScalableFloatDetails::GetRowNameComboBoxContentText)
+				]
+			]
+			+SHorizontalBox::Slot()
+			//.FillWidth(1.f)
+			.HAlign(HAlign_Fill)
+			.Padding(0.f, 0.f, 2.f, 0.f)
+			[
+				CurveTableProperty->CreatePropertyValueWidget()
+			]
+		];	
+}
+
+END_SLATE_FUNCTION_BUILD_OPTIMIZATION
+
+
+void FScalableFloatDetails::OnCurveTableChanged()
+{
+	CurrentSelectedItem = InitWidgetContent();
+	if (RowNameComboListView.IsValid())
+	{
+		RowNameComboListView->SetSelection(CurrentSelectedItem);
+		RowNameComboListView->RequestListRefresh();
+	}
+}
+
+TSharedPtr<FString> FScalableFloatDetails::InitWidgetContent()
+{
+	TSharedPtr<FString> InitialValue = MakeShareable(new FString(TEXT("None")));
+
+	FName RowName;
+	const FPropertyAccess::Result RowResult = RowNameProperty->GetValue(RowName);
+	RowNames.Empty();
+	RowNames.Add(InitialValue);
+
+	/** Get the properties we wish to work with */
+	UCurveTable* CurveTable = GetCurveTable();
+
+	if (CurveTable != NULL)
+	{
+		/** Extract all the row names from the RowMap */
+		for (TMap<FName, FRichCurve*>::TConstIterator Iterator(CurveTable->RowMap); Iterator; ++Iterator)
+		{
+			/** Create a simple array of the row names */
+			TSharedRef<FString> RowNameItem = MakeShareable(new FString(Iterator.Key().ToString()));
+			RowNames.Add(RowNameItem);
+
+			/** Set the initial value to the currently selected item */
+			if (Iterator.Key() == RowName)
+			{
+				InitialValue = RowNameItem;
+			}
+		}
+	}
+
+	/** Reset the initial value to ensure a valid entry is set */
+	if (RowResult != FPropertyAccess::MultipleValues)
+	{
+		FName NewValue = FName(**InitialValue);
+		RowNameProperty->SetValue(NewValue);
+	}
+
+	return InitialValue;
+}
+
+UCurveTable * FScalableFloatDetails::GetCurveTable()
+{
+	UCurveTable* CurveTable = NULL;
+	CurveTableProperty->GetValue((UObject*&)CurveTable);
+
+	if (CurveTable == NULL)
+	{
+		CurveTable = IGameplayAbilitiesModule::Get().GetSkillSystemGlobals()->GetGlobalCurveTable();
+	}
+
+	return CurveTable;
+}
+
+TSharedRef<SWidget> FScalableFloatDetails::GetListContent()
+{
+	SAssignNew(RowNameComboListView, SListView<TSharedPtr<FString> >)
+		.ListItemsSource(&RowNames)
+		.OnSelectionChanged(this, &FScalableFloatDetails::OnSelectionChanged)
+		.OnGenerateRow(this, &FScalableFloatDetails::HandleRowNameComboBoxGenarateWidget)
+		.SelectionMode(ESelectionMode::Single);
+
+	if (CurrentSelectedItem.IsValid())
+	{
+		RowNameComboListView->SetSelection(CurrentSelectedItem);
+	}
+
+	return SNew(SVerticalBox)
+		+ SVerticalBox::Slot()
+		.AutoHeight()
+		[
+			SNew(SSearchBox)
+			.OnTextChanged(this, &FScalableFloatDetails::OnFilterTextChanged)
+		]
+		+ SVerticalBox::Slot()
+		.FillHeight(1.f)
+		[
+			RowNameComboListView.ToSharedRef()
+		];
+}
+
+void FScalableFloatDetails::OnSelectionChanged(TSharedPtr<FString> SelectedItem, ESelectInfo::Type SelectInfo)
+{
+	if (SelectedItem.IsValid())
+	{
+		CurrentSelectedItem = SelectedItem;
+		FName NewValue = FName(**SelectedItem);
+		RowNameProperty->SetValue(NewValue);
+
+		// Close the combo
+		RowNameComboButton->SetIsOpen(false);
+	}
+}
+
+TSharedRef<ITableRow> FScalableFloatDetails::HandleRowNameComboBoxGenarateWidget(TSharedPtr<FString> InItem, const TSharedRef<STableViewBase>& OwnerTable)
+{
+	return
+		SNew(STableRow<TSharedPtr<FString>>, OwnerTable)
+		[
+			SNew(STextBlock).Text(*InItem)
+		];
+}
+
+/** Display the current selection */
+FString FScalableFloatDetails::GetRowNameComboBoxContentText() const
+{
+	FString RowName = TEXT("Multiple Values");
+	const FPropertyAccess::Result RowResult = RowNameProperty->GetValue(RowName);
+	if (RowResult != FPropertyAccess::MultipleValues)
+	{
+		TSharedPtr<FString> SelectedRowName = CurrentSelectedItem;
+		if (SelectedRowName.IsValid())
+		{
+			RowName = *SelectedRowName;
+		}
+		else
+		{
+			RowName = TEXT("None");
+		}
+	}
+	return RowName;
+}
+
+/** Called by Slate when the filter box changes text. */
+void FScalableFloatDetails::OnFilterTextChanged(const FText& InFilterText)
+{
+	FString CurrentFilterText = InFilterText.ToString();
+
+	FName RowName;
+	const FPropertyAccess::Result RowResult = RowNameProperty->GetValue(RowName);
+	RowNames.Empty();
+
+	/** Get the properties we wish to work with */
+	UCurveTable* CurveTable = GetCurveTable();
+
+	if (CurveTable != NULL)
+	{
+		/** Extract all the row names from the RowMap */
+		for (TMap<FName, FRichCurve*>::TConstIterator Iterator(CurveTable->RowMap); Iterator; ++Iterator)
+		{
+			/** Create a simple array of the row names */
+			FString RowString = Iterator.Key().ToString();
+			if (CurrentFilterText == TEXT("") || RowString.Contains(CurrentFilterText))
+			{
+				TSharedRef<FString> RowNameItem = MakeShareable(new FString(RowString));
+				RowNames.Add(RowNameItem);
+			}
+		}
+	}
+	RowNameComboListView->RequestListRefresh();
+}
+
+bool FScalableFloatDetails::IsEditable( ) const
+{
+	return true;
+}
+
+void FScalableFloatDetails::CustomizeChildren( TSharedRef<class IPropertyHandle> StructPropertyHandle, class IDetailChildrenBuilder& StructBuilder, IStructCustomizationUtils& StructCustomizationUtils )
+{
+	
+}
+
+//-------------------------------------------------------------------------------------
+
+
+#undef LOCTEXT_NAMESPACE
\ No newline at end of file
diff --git a/Engine/Source/Editor/GameplayAbilitiesEditor/Private/AttributeDetails.h b/Engine/Source/Editor/GameplayAbilitiesEditor/Private/AttributeDetails.h
new file mode 100644
--- /dev/null
+++ b/Engine/Source/Editor/GameplayAbilitiesEditor/Private/AttributeDetails.h
@@ -0,0 +1,86 @@
+// Copyright 1998-2013 Epic Games, Inc. All Rights Reserved.
+
+#pragma once
+
+#include "Editor/PropertyEditor/Public/PropertyEditorModule.h"
+
+class IPropertyHandle;
+class FDetailWidgetRow;
+class IDetailLayoutBuilder;
+
+DECLARE_LOG_CATEGORY_EXTERN(LogAttributeDetails, Log, All);
+
+class FAttributeDetails : public IDetailCustomization
+{
+public:
+	/** Makes a new instance of this detail layout class for a specific detail view requesting it */
+	static TSharedRef<IDetailCustomization> MakeInstance();
+
+private:
+	/** IDetailCustomization interface */
+	virtual void CustomizeDetails(IDetailLayoutBuilder& DetailLayout) OVERRIDE;
+
+	TSharedPtr<FString> GetPropertyType() const;
+
+	TArray<TSharedPtr<FString>> PropertyOptions;
+
+	TSharedPtr<IPropertyHandle> MyProperty;
+
+	void OnChangeProperty(TSharedPtr<FString> ItemSelected, ESelectInfo::Type SelectInfo);	
+};
+
+
+class FAttributePropertyDetails : public IPropertyTypeCustomization
+{
+public:
+	static TSharedRef<IPropertyTypeCustomization> MakeInstance();
+
+	/** IPropertyTypeCustomization interface */
+	virtual void CustomizeHeader( TSharedRef<class IPropertyHandle> StructPropertyHandle, class FDetailWidgetRow& HeaderRow, IStructCustomizationUtils& StructCustomizationUtils ) OVERRIDE;
+	virtual void CustomizeChildren( TSharedRef<class IPropertyHandle> StructPropertyHandle, class IDetailChildrenBuilder& StructBuilder, IStructCustomizationUtils& StructCustomizationUtils ) OVERRIDE;              
+
+private:
+
+	TSharedPtr<IPropertyHandle> MyProperty;
+	TArray<TSharedPtr<FString>> PropertyOptions;
+
+	TSharedPtr<FString>	GetPropertyType() const;
+
+	void OnChangeProperty(TSharedPtr<FString> ItemSelected, ESelectInfo::Type SelectInfo);
+};
+
+class FScalableFloatDetails : public IPropertyTypeCustomization
+{
+public:
+	static TSharedRef<IPropertyTypeCustomization> MakeInstance();
+
+
+private:
+
+	virtual void CustomizeHeader( TSharedRef<class IPropertyHandle> StructPropertyHandle, class FDetailWidgetRow& HeaderRow, IStructCustomizationUtils& StructCustomizationUtils ) OVERRIDE;
+	virtual void CustomizeChildren( TSharedRef<class IPropertyHandle> StructPropertyHandle, class IDetailChildrenBuilder& StructBuilder, IStructCustomizationUtils& StructCustomizationUtils ) OVERRIDE;              
+
+	bool IsEditable( ) const;
+
+	TSharedRef<SWidget> GetListContent();
+	void OnSelectionChanged(TSharedPtr<FString> SelectedItem, ESelectInfo::Type SelectInfo);
+	FString GetRowNameComboBoxContentText() const;
+	TSharedRef<ITableRow> HandleRowNameComboBoxGenarateWidget(TSharedPtr<FString> InItem, const TSharedRef<STableViewBase>& OwnerTable);
+	void OnFilterTextChanged(const FText& InFilterText);
+
+	class UCurveTable * GetCurveTable();
+
+	TSharedPtr<FString> InitWidgetContent();
+
+	void OnCurveTableChanged();
+
+	TSharedPtr<SComboButton> RowNameComboButton;
+	TSharedPtr<SListView<TSharedPtr<FString> > > RowNameComboListView;
+	TSharedPtr<FString> CurrentSelectedItem;
+	TArray<TSharedPtr<FString> > RowNames;
+
+	TSharedPtr<IPropertyHandle> ValueProperty;
+	TSharedPtr<IPropertyHandle> CurveTableHandleProperty;
+	TSharedPtr<IPropertyHandle> CurveTableProperty;
+	TSharedPtr<IPropertyHandle> RowNameProperty;
+};
\ No newline at end of file
diff --git a/Engine/Source/Editor/GameplayAbilitiesEditor/Private/GameplayAbilitiesEditorModule.cpp b/Engine/Source/Editor/GameplayAbilitiesEditor/Private/GameplayAbilitiesEditorModule.cpp
new file mode 100644
--- /dev/null
+++ b/Engine/Source/Editor/GameplayAbilitiesEditor/Private/GameplayAbilitiesEditorModule.cpp
@@ -0,0 +1,32 @@
+// Copyright 1998-2013 Epic Games, Inc. All Rights Reserved.
+
+#include "SkillSystemEditorModulePrivatePCH.h"
+#include "Editor/PropertyEditor/Public/PropertyEditorModule.h"
+#include "AttributeDetails.h"
+#include "AttributeSet.h"
+
+class FGameplayAbilitiesEditorModule : public IGameplayAbilitiesEditorModule
+{
+	// Begin IModuleInterface
+	virtual void StartupModule() OVERRIDE;
+	virtual void ShutdownModule() OVERRIDE;
+	// End IModuleInterface
+};
+
+IMPLEMENT_MODULE(FGameplayAbilitiesEditorModule, GameplayAbilitiesEditor)
+
+void FGameplayAbilitiesEditorModule::StartupModule()
+{
+	// Register the details customizer
+	FPropertyEditorModule& PropertyModule = FModuleManager::LoadModuleChecked<FPropertyEditorModule>("PropertyEditor");
+	PropertyModule.RegisterCustomPropertyTypeLayout( "GameplayAttribute", FOnGetPropertyTypeCustomizationInstance::CreateStatic( &FAttributePropertyDetails::MakeInstance ) );
+	PropertyModule.RegisterCustomPropertyTypeLayout( "ScalableFloat", FOnGetPropertyTypeCustomizationInstance::CreateStatic( &FScalableFloatDetails::MakeInstance ) );
+
+	PropertyModule.RegisterCustomPropertyLayout( "AttributeSet", FOnGetDetailCustomizationInstance::CreateStatic( &FAttributeDetails::MakeInstance ) );
+}
+
+void FGameplayAbilitiesEditorModule::ShutdownModule()
+{
+	// This function may be called during shutdown to clean up your module.  For modules that support dynamic reloading,
+	// we call this function before unloading the module.
+}
diff --git a/Engine/Source/Editor/GameplayAbilitiesEditor/Private/K2Node_PlayMontageAndWait.cpp b/Engine/Source/Editor/GameplayAbilitiesEditor/Private/K2Node_PlayMontageAndWait.cpp
new file mode 100644
--- /dev/null
+++ b/Engine/Source/Editor/GameplayAbilitiesEditor/Private/K2Node_PlayMontageAndWait.cpp
@@ -0,0 +1,38 @@
+// Copyright 1998-2014 Epic Games, Inc. All Rights Reserved.
+
+
+#include "SkillSystemEditorModulePrivatePCH.h"
+#include "K2Node_PlayMontageAndWait.h"
+#include "SkillSystemBlueprintLibrary.h"
+#include "Abilities/Tasks/AbilityTask_PlayMontageAndWait.h"
+#include "CompilerResultsLog.h"
+
+#define LOCTEXT_NAMESPACE "K2Node_PlayMontageAndWait"
+
+UK2Node_PlayMontageAndWait::UK2Node_PlayMontageAndWait(const class FPostConstructInitializeProperties& PCIP)
+	: Super(PCIP)
+{
+	
+	ProxyFactoryFunctionName = GET_FUNCTION_NAME_CHECKED(USkillSystemBlueprintLibrary, CreatePlayMontageAndWaitProxy);
+	ProxyFactoryClass = USkillSystemBlueprintLibrary::StaticClass();
+	ProxyClass = UAbilityTask_PlayMontageAndWait::StaticClass();
+}
+
+FString UK2Node_PlayMontageAndWait::GetCategoryName()
+{
+	return TEXT("Ability");
+}
+
+FString UK2Node_PlayMontageAndWait::GetTooltip() const
+{
+	return TEXT("Play Anim Montage until it completes or is interrupted");
+}
+
+FText UK2Node_PlayMontageAndWait::GetNodeTitle(ENodeTitleType::Type TitleType) const
+{
+	return LOCTEXT("PlayMontageAndWait", "Play Montage and Wait");
+}
+
+#undef LOCTEXT_NAMESPACE
+
+
diff --git a/Engine/Source/Editor/GameplayAbilitiesEditor/Private/K2Node_WaitMovementModeChange.cpp b/Engine/Source/Editor/GameplayAbilitiesEditor/Private/K2Node_WaitMovementModeChange.cpp
new file mode 100644
--- /dev/null
+++ b/Engine/Source/Editor/GameplayAbilitiesEditor/Private/K2Node_WaitMovementModeChange.cpp
@@ -0,0 +1,38 @@
+// Copyright 1998-2014 Epic Games, Inc. All Rights Reserved.
+
+
+#include "SkillSystemEditorModulePrivatePCH.h"
+#include "K2Node_WaitMovementModeChange.h"
+#include "SkillSystemBlueprintLibrary.h"
+#include "Abilities/Tasks/AbilityTask_WaitMovementModeChange.h"
+#include "CompilerResultsLog.h"
+
+#define LOCTEXT_NAMESPACE "K2Node_WaitMovementModeChange"
+
+UK2Node_WaitMovementModeChange::UK2Node_WaitMovementModeChange(const class FPostConstructInitializeProperties& PCIP)
+	: Super(PCIP)
+{
+	
+	ProxyFactoryFunctionName = GET_FUNCTION_NAME_CHECKED(USkillSystemBlueprintLibrary, CreateWaitMovementModeChange);
+	ProxyFactoryClass = USkillSystemBlueprintLibrary::StaticClass();
+	ProxyClass = UAbilityTask_WaitMovementModeChange::StaticClass();
+}
+
+FString UK2Node_WaitMovementModeChange::GetCategoryName()
+{
+	return TEXT("Ability");
+}
+
+FString UK2Node_WaitMovementModeChange::GetTooltip() const
+{
+	return TEXT("Waits until the movement mode changes");
+}
+
+FText UK2Node_WaitMovementModeChange::GetNodeTitle(ENodeTitleType::Type TitleType) const
+{
+	return LOCTEXT("WaitMovementModeChange", "Wait Movement Mode Change");
+}
+
+#undef LOCTEXT_NAMESPACE
+
+
diff --git a/Engine/Source/Editor/GameplayAbilitiesEditor/Private/K2Node_WaitOverlap.cpp b/Engine/Source/Editor/GameplayAbilitiesEditor/Private/K2Node_WaitOverlap.cpp
new file mode 100644
--- /dev/null
+++ b/Engine/Source/Editor/GameplayAbilitiesEditor/Private/K2Node_WaitOverlap.cpp
@@ -0,0 +1,38 @@
+// Copyright 1998-2014 Epic Games, Inc. All Rights Reserved.
+
+
+#include "SkillSystemEditorModulePrivatePCH.h"
+#include "K2Node_WaitOverlap.h"
+#include "SkillSystemBlueprintLibrary.h"
+#include "Abilities/Tasks/AbilityTask_WaitOverlap.h"
+#include "CompilerResultsLog.h"
+
+#define LOCTEXT_NAMESPACE "K2Node_WaitOverlap"
+
+UK2Node_WaitOverlap::UK2Node_WaitOverlap(const class FPostConstructInitializeProperties& PCIP)
+	: Super(PCIP)
+{
+	
+	ProxyFactoryFunctionName = GET_FUNCTION_NAME_CHECKED(USkillSystemBlueprintLibrary, CreateWaitOverlap);
+	ProxyFactoryClass = USkillSystemBlueprintLibrary::StaticClass();
+	ProxyClass = UAbilityTask_WaitOverlap::StaticClass();
+}
+
+FString UK2Node_WaitOverlap::GetCategoryName()
+{
+	return TEXT("Ability");
+}
+
+FString UK2Node_WaitOverlap::GetTooltip() const
+{
+	return TEXT("Wait until an overlap occurs");
+}
+
+FText UK2Node_WaitOverlap::GetNodeTitle(ENodeTitleType::Type TitleType) const
+{
+	return LOCTEXT("WaitForOverlap", "Wait For Overlap");
+}
+
+#undef LOCTEXT_NAMESPACE
+
+
diff --git a/Engine/Source/Editor/GameplayAbilitiesEditor/Private/SkillSystemEditorModulePrivatePCH.h b/Engine/Source/Editor/GameplayAbilitiesEditor/Private/SkillSystemEditorModulePrivatePCH.h
new file mode 100644
--- /dev/null
+++ b/Engine/Source/Editor/GameplayAbilitiesEditor/Private/SkillSystemEditorModulePrivatePCH.h
@@ -0,0 +1,15 @@
+// Copyright 1998-2013 Epic Games, Inc. All Rights Reserved.
+
+#include "GameplayAbilitiesEditorModule.h"
+
+#if WITH_EDITOR
+#include "UnrealEd.h"
+#endif
+
+
+// You should place include statements to your module's private header files here.  You only need to
+// add includes for headers that are used in most of your module's source files though.
+
+#include "Core.h"
+#include "Engine.h"
+#include "BlueprintGraphDefinitions.h"
\ No newline at end of file
diff --git a/Engine/Source/Editor/GameplayAbilitiesEditor/Public/GameplayAbilitiesEditorModule.h b/Engine/Source/Editor/GameplayAbilitiesEditor/Public/GameplayAbilitiesEditorModule.h
new file mode 100644
--- /dev/null
+++ b/Engine/Source/Editor/GameplayAbilitiesEditor/Public/GameplayAbilitiesEditorModule.h
@@ -0,0 +1,37 @@
+// Copyright 1998-2013 Epic Games, Inc. All Rights Reserved.
+
+#pragma once
+
+#include "ModuleManager.h"
+
+
+/**
+ * The public interface to this module
+ */
+class IGameplayAbilitiesEditorModule : public IModuleInterface
+{
+
+public:
+
+	/**
+	 * Singleton-like access to this module's interface.  This is just for convenience!
+	 * Beware of calling this during the shutdown phase, though.  Your module might have been unloaded already.
+	 *
+	 * @return Returns singleton instance, loading the module on demand if needed
+	 */
+	static inline IGameplayAbilitiesEditorModule& Get()
+	{
+		return FModuleManager::LoadModuleChecked< IGameplayAbilitiesEditorModule >("GameplayAbilitiesEditor");
+	}
+
+	/**
+	 * Checks to see if this module is loaded and ready.  It is only valid to call Get() if IsAvailable() returns true.
+	 *
+	 * @return True if the module is loaded and ready to use
+	 */
+	static inline bool IsAvailable()
+	{
+		return FModuleManager::Get().IsModuleLoaded( "SkillSystemEditor" );
+	}
+};
+
diff --git a/Engine/Source/Editor/SkillSystemEditor/Classes/K2Node_PlayMontageAndWait.h b/Engine/Source/Editor/SkillSystemEditor/Classes/K2Node_PlayMontageAndWait.h
deleted file mode 100644
diff --git a/Engine/Source/Editor/SkillSystemEditor/Classes/K2Node_WaitMovementModeChange.h b/Engine/Source/Editor/SkillSystemEditor/Classes/K2Node_WaitMovementModeChange.h
deleted file mode 100644
diff --git a/Engine/Source/Editor/SkillSystemEditor/Classes/K2Node_WaitOverlap.h b/Engine/Source/Editor/SkillSystemEditor/Classes/K2Node_WaitOverlap.h
deleted file mode 100644
diff --git a/Engine/Source/Editor/SkillSystemEditor/Private/AttributeDetails.cpp b/Engine/Source/Editor/SkillSystemEditor/Private/AttributeDetails.cpp
deleted file mode 100644
diff --git a/Engine/Source/Editor/SkillSystemEditor/Private/AttributeDetails.h b/Engine/Source/Editor/SkillSystemEditor/Private/AttributeDetails.h
deleted file mode 100644
diff --git a/Engine/Source/Editor/SkillSystemEditor/Private/K2Node_PlayMontageAndWait.cpp b/Engine/Source/Editor/SkillSystemEditor/Private/K2Node_PlayMontageAndWait.cpp
deleted file mode 100644
diff --git a/Engine/Source/Editor/SkillSystemEditor/Private/K2Node_WaitMovementModeChange.cpp b/Engine/Source/Editor/SkillSystemEditor/Private/K2Node_WaitMovementModeChange.cpp
deleted file mode 100644
diff --git a/Engine/Source/Editor/SkillSystemEditor/Private/K2Node_WaitOverlap.cpp b/Engine/Source/Editor/SkillSystemEditor/Private/K2Node_WaitOverlap.cpp
deleted file mode 100644
diff --git a/Engine/Source/Editor/SkillSystemEditor/Private/SkillSystemEditorModule.cpp b/Engine/Source/Editor/SkillSystemEditor/Private/SkillSystemEditorModule.cpp
deleted file mode 100644
diff --git a/Engine/Source/Editor/SkillSystemEditor/Private/SkillSystemEditorModulePrivatePCH.h b/Engine/Source/Editor/SkillSystemEditor/Private/SkillSystemEditorModulePrivatePCH.h
deleted file mode 100644
diff --git a/Engine/Source/Editor/SkillSystemEditor/Public/SkillSystemEditorModule.h b/Engine/Source/Editor/SkillSystemEditor/Public/SkillSystemEditorModule.h
deleted file mode 100644
diff --git a/Engine/Source/Editor/SkillSystemEditor/SkillSystemEditor.Build.cs b/Engine/Source/Editor/SkillSystemEditor/SkillSystemEditor.Build.cs
deleted file mode 100644
