First pass of the abilities editor.

diff --git a/Engine/Source/Editor/GameplayAbilitiesEditor/Classes/GameplayAbilitiesBlueprintFactory.h b/Engine/Source/Editor/GameplayAbilitiesEditor/Classes/GameplayAbilitiesBlueprintFactory.h
new file mode 100644
--- /dev/null
+++ b/Engine/Source/Editor/GameplayAbilitiesEditor/Classes/GameplayAbilitiesBlueprintFactory.h
@@ -0,0 +1,25 @@
+// Copyright 1998-2014 Epic Games, Inc. All Rights Reserved.
+
+#pragma once
+
+#include "GameplayAbilitiesBlueprintFactory.generated.h"
+
+UCLASS(HideCategories=Object, MinimalAPI)
+class UGameplayAbilitiesBlueprintFactory : public UFactory
+{
+	GENERATED_UCLASS_BODY()
+
+	// The type of blueprint that will be created
+	UPROPERTY(EditAnywhere, Category=GameplayAbilitiesBlueprintFactory)
+	TEnumAsByte<enum EBlueprintType> BlueprintType;
+
+	// The parent class of the created blueprint
+	UPROPERTY(EditAnywhere, Category=GameplayAbilitiesBlueprintFactory)
+	TSubclassOf<class UGameplayAbility> ParentClass;
+
+	// Begin UFactory Interface
+	virtual bool ConfigureProperties() override;
+	virtual UObject* FactoryCreateNew(UClass* Class, UObject* InParent, FName Name, EObjectFlags Flags, UObject* Context, FFeedbackContext* Warn, FName CallingContext) override;
+	virtual UObject* FactoryCreateNew(UClass* Class, UObject* InParent, FName Name, EObjectFlags Flags, UObject* Context, FFeedbackContext* Warn) override;
+	// Begin UFactory Interface	
+};
diff --git a/Engine/Source/Editor/GameplayAbilitiesEditor/Classes/GameplayAbilityGraphSchema.h b/Engine/Source/Editor/GameplayAbilitiesEditor/Classes/GameplayAbilityGraphSchema.h
new file mode 100644
--- /dev/null
+++ b/Engine/Source/Editor/GameplayAbilitiesEditor/Classes/GameplayAbilityGraphSchema.h
@@ -0,0 +1,14 @@
+// Copyright 1998-2014 Epic Games, Inc. All Rights Reserved.
+
+#pragma once
+
+#include "EdGraphSchema_K2.h"
+#include "GameplayAbilityGraphSchema.generated.h"
+
+UCLASS(MinimalAPI)
+class UGameplayAbilityGraphSchema : public UEdGraphSchema_K2
+{
+	GENERATED_UCLASS_BODY()
+
+	// Dummy class needed for possible extensions later
+};
\ No newline at end of file
diff --git a/Engine/Source/Editor/GameplayAbilitiesEditor/GameplayAbilitiesEditor.Build.cs b/Engine/Source/Editor/GameplayAbilitiesEditor/GameplayAbilitiesEditor.Build.cs
--- a/Engine/Source/Editor/GameplayAbilitiesEditor/GameplayAbilitiesEditor.Build.cs
+++ b/Engine/Source/Editor/GameplayAbilitiesEditor/GameplayAbilitiesEditor.Build.cs
@@ -26,6 +26,7 @@ namespace UnrealBuildTool.Rules
 					"SlateCore",
                     "EditorStyle",
 					"BlueprintGraph",
+                    "Kismet",
 					"KismetCompiler",
 					"GraphEditor",
 					"MainFrame",
diff --git a/Engine/Source/Editor/GameplayAbilitiesEditor/Private/AssetTypeActions_GameplayAbilitiesBlueprint.cpp b/Engine/Source/Editor/GameplayAbilitiesEditor/Private/AssetTypeActions_GameplayAbilitiesBlueprint.cpp
new file mode 100644
--- /dev/null
+++ b/Engine/Source/Editor/GameplayAbilitiesEditor/Private/AssetTypeActions_GameplayAbilitiesBlueprint.cpp
@@ -0,0 +1,62 @@
+// Copyright 1998-2014 Epic Games, Inc. All Rights Reserved.
+
+#include "AbilitySystemEditorPrivatePCH.h"
+
+#include "PackageTools.h"
+#include "Kismet2/BlueprintEditorUtils.h"
+#include "Kismet2/KismetEditorUtilities.h"
+#include "BlueprintEditorModule.h"
+#include "AssetRegistryModule.h"
+#include "SBlueprintDiff.h"
+#include "ISourceControlModule.h"
+#include "MessageLog.h"
+
+#include "AssetTypeActions_GameplayAbilitiesBlueprint.h"
+#include "GameplayAbilitiesEditor.h"
+#include "GameplayAbilityBlueprint.h"
+
+#define LOCTEXT_NAMESPACE "AssetTypeActions"
+
+FText FAssetTypeActions_GameplayAbilitiesBlueprint::GetName() const
+{ 
+	return NSLOCTEXT("AssetTypeActions", "AssetTypeActions_GameplayAbilitiesBlueprint", "Gameplay Ability Blueprint"); 
+}
+
+void FAssetTypeActions_GameplayAbilitiesBlueprint::OpenAssetEditor( const TArray<UObject*>& InObjects, TSharedPtr<IToolkitHost> EditWithinLevelEditor )
+{
+	EToolkitMode::Type Mode = EditWithinLevelEditor.IsValid() ? EToolkitMode::WorldCentric : EToolkitMode::Standalone;
+
+	for (auto ObjIt = InObjects.CreateConstIterator(); ObjIt; ++ObjIt)
+	{
+		auto Blueprint = Cast<UBlueprint>(*ObjIt);
+		if (Blueprint && Blueprint->SkeletonGeneratedClass && Blueprint->GeneratedClass )
+		{
+			TSharedRef< FGameplayAbilitiesEditor > NewEditor(new FGameplayAbilitiesEditor());
+
+			TArray<UBlueprint*> Blueprints;
+			Blueprints.Add(Blueprint);
+
+			NewEditor->InitGameplayAbilitiesEditor(Mode, EditWithinLevelEditor, Blueprints, ShouldUseDataOnlyEditor(Blueprint));
+		}
+		else
+		{
+			FMessageDialog::Open( EAppMsgType::Ok, LOCTEXT("FailedToLoadBlueprint", "Gameplay Ability Blueprint could not be loaded because it derives from an invalid class.  Check to make sure the parent class for this blueprint hasn't been removed!"));
+		}
+	}
+}
+
+bool FAssetTypeActions_GameplayAbilitiesBlueprint::ShouldUseDataOnlyEditor(const UBlueprint* Blueprint) const
+{
+	return FBlueprintEditorUtils::IsDataOnlyBlueprint(Blueprint)
+		&& !FBlueprintEditorUtils::IsLevelScriptBlueprint(Blueprint)
+		&& !FBlueprintEditorUtils::IsInterfaceBlueprint(Blueprint)
+		&& !Blueprint->bForceFullEditor
+		&& !Blueprint->bIsNewlyCreated;
+}
+
+UClass* FAssetTypeActions_GameplayAbilitiesBlueprint::GetSupportedClass() const
+{ 
+	return UGameplayAbilityBlueprint::StaticClass(); 
+}
+
+#undef LOCTEXT_NAMESPACE
\ No newline at end of file
diff --git a/Engine/Source/Editor/GameplayAbilitiesEditor/Private/AssetTypeActions_GameplayAbilitiesBlueprint.h b/Engine/Source/Editor/GameplayAbilitiesEditor/Private/AssetTypeActions_GameplayAbilitiesBlueprint.h
new file mode 100644
--- /dev/null
+++ b/Engine/Source/Editor/GameplayAbilitiesEditor/Private/AssetTypeActions_GameplayAbilitiesBlueprint.h
@@ -0,0 +1,22 @@
+// Copyright 1998-2014 Epic Games, Inc. All Rights Reserved.
+
+#pragma once
+
+#include "AssetTypeActions_Base.h"
+
+//TODO Use base class FAssetTypeActions_Blueprint later
+class FAssetTypeActions_GameplayAbilitiesBlueprint : public FAssetTypeActions_Base
+{
+public:
+	// IAssetTypeActions Implementation
+	virtual FText GetName() const override;
+	virtual FColor GetTypeColor() const override { return FColor(255,255,255); }
+	virtual UClass* GetSupportedClass() const override;
+	virtual void OpenAssetEditor(const TArray<UObject*>& InObjects, TSharedPtr<class IToolkitHost> EditWithinLevelEditor = TSharedPtr<IToolkitHost>()) override;
+	virtual uint32 GetCategories() override { return EAssetTypeCategories::Gameplay; }
+	// End IAssetTypeActions Implementation
+
+private:
+	/** Returns true if the blueprint is data only */
+	bool ShouldUseDataOnlyEditor(const UBlueprint* Blueprint) const;
+};
\ No newline at end of file
diff --git a/Engine/Source/Editor/GameplayAbilitiesEditor/Private/GameplayAbilitiesBlueprintFactory.cpp b/Engine/Source/Editor/GameplayAbilitiesEditor/Private/GameplayAbilitiesBlueprintFactory.cpp
new file mode 100644
--- /dev/null
+++ b/Engine/Source/Editor/GameplayAbilitiesEditor/Private/GameplayAbilitiesBlueprintFactory.cpp
@@ -0,0 +1,65 @@
+// Copyright 1998-2014 Epic Games, Inc. All Rights Reserved.
+
+#include "AbilitySystemEditorPrivatePCH.h"
+
+#include "Kismet2/BlueprintEditorUtils.h"
+#include "Kismet2/KismetEditorUtilities.h"
+
+#include "Abilities/GameplayAbility.h"
+
+#include "GameplayAbilityBlueprintGeneratedClass.h"
+
+#include "GameplayAbilitiesBlueprintFactory.h"
+
+#define LOCTEXT_NAMESPACE "UGameplayAbilitiesBlueprintFactory"
+
+/*------------------------------------------------------------------------------
+	UGameplayAbilitiesBlueprintFactory implementation.
+------------------------------------------------------------------------------*/
+
+UGameplayAbilitiesBlueprintFactory::UGameplayAbilitiesBlueprintFactory(const class FPostConstructInitializeProperties& PCIP)
+	: Super(PCIP)
+{
+	bCreateNew = true;
+	bEditAfterNew = true;
+	SupportedClass = UGameplayAbilityBlueprint::StaticClass();
+	ParentClass = UGameplayAbility::StaticClass();
+}
+
+bool UGameplayAbilitiesBlueprintFactory::ConfigureProperties()
+{
+	return true;
+};
+
+UObject* UGameplayAbilitiesBlueprintFactory::FactoryCreateNew(UClass* Class, UObject* InParent, FName Name, EObjectFlags Flags, UObject* Context, FFeedbackContext* Warn, FName CallingContext)
+{
+	// Make sure we are trying to factory a gameplay ability blueprint, then create and init one
+	check(Class->IsChildOf(UGameplayAbilityBlueprint::StaticClass()));
+
+	// If they selected an interface, force the parent class to be UInterface
+	if (BlueprintType == BPTYPE_Interface)
+	{
+		ParentClass = UInterface::StaticClass();
+	}
+
+	if ( ( ParentClass == NULL ) || !FKismetEditorUtilities::CanCreateBlueprintOfClass(ParentClass) || !ParentClass->IsChildOf(UGameplayAbility::StaticClass()) )
+	{
+		FFormatNamedArguments Args;
+		Args.Add( TEXT("ClassName"), (ParentClass != NULL) ? FText::FromString( ParentClass->GetName() ) : LOCTEXT("Null", "(null)") );
+		FMessageDialog::Open( EAppMsgType::Ok, FText::Format( LOCTEXT("CannotCreateGameplayAbilityBlueprint", "Cannot create a Gameplay Ability Blueprint based on the class '{ClassName}'."), Args ) );
+		return NULL;
+	}
+	else
+	{
+		UGameplayAbilityBlueprint* NewBP = CastChecked<UGameplayAbilityBlueprint>(FKismetEditorUtilities::CreateBlueprint(ParentClass, InParent, Name, BlueprintType, UGameplayAbilityBlueprint::StaticClass(), UGameplayAbilityBlueprintGeneratedClass::StaticClass(), CallingContext));
+
+		return NewBP;
+	}
+}
+
+UObject* UGameplayAbilitiesBlueprintFactory::FactoryCreateNew(UClass* Class, UObject* InParent, FName Name, EObjectFlags Flags, UObject* Context, FFeedbackContext* Warn)
+{
+	return FactoryCreateNew(Class, InParent, Name, Flags, Context, Warn, NAME_None);
+}
+
+#undef LOCTEXT_NAMESPACE
diff --git a/Engine/Source/Editor/GameplayAbilitiesEditor/Private/GameplayAbilitiesEditor.cpp b/Engine/Source/Editor/GameplayAbilitiesEditor/Private/GameplayAbilitiesEditor.cpp
new file mode 100644
--- /dev/null
+++ b/Engine/Source/Editor/GameplayAbilitiesEditor/Private/GameplayAbilitiesEditor.cpp
@@ -0,0 +1,83 @@
+// Copyright 1998-2014 Epic Games, Inc. All Rights Reserved.
+
+#include "AbilitySystemEditorPrivatePCH.h"
+#include "GameplayAbilitiesEditor.h"
+
+#define LOCTEXT_NAMESPACE "FGameplayAbilitiesEditor"
+
+
+/////////////////////////////////////////////////////
+// FGameplayAbilitiesEditor
+
+FGameplayAbilitiesEditor::FGameplayAbilitiesEditor()
+{
+	// todo: Do we need to register a callback for when properties are changed?
+}
+
+FGameplayAbilitiesEditor::~FGameplayAbilitiesEditor()
+{
+	FEditorDelegates::OnAssetPostImport.RemoveAll(this);
+	FReimportManager::Instance()->OnPostReimport().RemoveAll(this);
+	
+	// NOTE: Any tabs that we still have hanging out when destroyed will be cleaned up by FBaseToolkit's destructor
+}
+
+void FGameplayAbilitiesEditor::InitGameplayAbilitiesEditor(const EToolkitMode::Type Mode, const TSharedPtr< IToolkitHost >& InitToolkitHost, const TArray<UBlueprint*>& InBlueprints, bool bShouldOpenInDefaultsMode)
+{
+	InitBlueprintEditor(Mode, InitToolkitHost, InBlueprints, bShouldOpenInDefaultsMode);
+}
+
+FName FGameplayAbilitiesEditor::GetToolkitFName() const
+{
+	return FName("GameplayAbilitiesEditor");
+}
+
+FText FGameplayAbilitiesEditor::GetBaseToolkitName() const
+{
+	return LOCTEXT("GameplayAbilitiesEditorAppLabel", "Gameplay Abilities Editor");
+}
+
+FText FGameplayAbilitiesEditor::GetToolkitName() const
+{
+	const auto EditingObjects = GetEditingObjects();
+
+	check(EditingObjects.Num() > 0);
+
+	FFormatNamedArguments Args;
+
+	const UObject* EditingObject = EditingObjects[0];
+
+	const bool bDirtyState = EditingObject->GetOutermost()->IsDirty();
+
+	Args.Add(TEXT("ObjectName"), FText::FromString(EditingObject->GetName()));
+	Args.Add(TEXT("DirtyState"), bDirtyState ? FText::FromString(TEXT("*")) : FText::GetEmpty());
+	return FText::Format(LOCTEXT("GameplayAbilitiesEditorAppLabel", "{ObjectName}{DirtyState}"), Args);
+}
+
+FString FGameplayAbilitiesEditor::GetWorldCentricTabPrefix() const
+{
+	return TEXT("GameplayAbilitiesEditor");
+}
+
+FLinearColor FGameplayAbilitiesEditor::GetWorldCentricTabColorScale() const
+{
+	return FLinearColor::White;
+}
+
+UBlueprint* FGameplayAbilitiesEditor::GetBlueprintObj() const
+{
+	auto EditingObjects = GetEditingObjects();
+	for (int32 i = 0; i < EditingObjects.Num(); ++i)
+	{
+		if (EditingObjects[i]->IsA<UGameplayAbilityBlueprint>()) { return (UBlueprint*)EditingObjects[i]; }
+	}
+	return NULL;
+}
+
+FString FGameplayAbilitiesEditor::GetDocumentationLink() const
+{
+	return FBlueprintEditor::GetDocumentationLink(); // todo: point this at the correct documentation
+}
+
+#undef LOCTEXT_NAMESPACE
+
diff --git a/Engine/Source/Editor/GameplayAbilitiesEditor/Private/GameplayAbilitiesEditorModule.cpp b/Engine/Source/Editor/GameplayAbilitiesEditor/Private/GameplayAbilitiesEditorModule.cpp
--- a/Engine/Source/Editor/GameplayAbilitiesEditor/Private/GameplayAbilitiesEditorModule.cpp
+++ b/Engine/Source/Editor/GameplayAbilitiesEditor/Private/GameplayAbilitiesEditorModule.cpp
@@ -11,6 +11,13 @@ class FGameplayAbilitiesEditorModule : public IGameplayAbilitiesEditorModule
 	virtual void StartupModule() override;
 	virtual void ShutdownModule() override;
 	// End IModuleInterface
+
+protected:
+	void RegisterAssetTypeAction(class IAssetTools& AssetTools, TSharedRef<IAssetTypeActions> Action);
+
+private:
+	/** All created asset type actions.  Cached here so that we can unregister it during shutdown. */
+	TArray< TSharedPtr<IAssetTypeActions> > CreatedAssetTypeActions;
 };
 
 IMPLEMENT_MODULE(FGameplayAbilitiesEditorModule, GameplayAbilitiesEditor)
@@ -23,10 +30,24 @@ void FGameplayAbilitiesEditorModule::StartupModule()
 	PropertyModule.RegisterCustomPropertyTypeLayout( "ScalableFloat", FOnGetPropertyTypeCustomizationInstance::CreateStatic( &FScalableFloatDetails::MakeInstance ) );
 
 	PropertyModule.RegisterCustomClassLayout( "AttributeSet", FOnGetDetailCustomizationInstance::CreateStatic( &FAttributeDetails::MakeInstance ) );
+
+	// Register asset types
+	IAssetTools& AssetTools = FModuleManager::LoadModuleChecked<FAssetToolsModule>("AssetTools").Get();
+	TSharedRef<IAssetTypeActions> Action = MakeShareable(new FAssetTypeActions_GameplayAbilitiesBlueprint());
+	AssetTools.RegisterAssetTypeActions(Action);
+	CreatedAssetTypeActions.Add(Action);
+}
+
+void FGameplayAbilitiesEditorModule::RegisterAssetTypeAction(IAssetTools& AssetTools, TSharedRef<IAssetTypeActions> Action)
+{
+	AssetTools.RegisterAssetTypeActions(Action);
+	CreatedAssetTypeActions.Add(Action);
 }
 
 void FGameplayAbilitiesEditorModule::ShutdownModule()
 {
 	// This function may be called during shutdown to clean up your module.  For modules that support dynamic reloading,
 	// we call this function before unloading the module.
+
+	CreatedAssetTypeActions.Empty();
 }
diff --git a/Engine/Source/Editor/GameplayAbilitiesEditor/Private/GameplayAbilityGraphSchema.cpp b/Engine/Source/Editor/GameplayAbilitiesEditor/Private/GameplayAbilityGraphSchema.cpp
new file mode 100644
--- /dev/null
+++ b/Engine/Source/Editor/GameplayAbilitiesEditor/Private/GameplayAbilityGraphSchema.cpp
@@ -0,0 +1,10 @@
+// Copyright 1998-2014 Epic Games, Inc. All Rights Reserved.
+
+#include "AbilitySystemEditorPrivatePCH.h"
+
+#include "GameplayAbilityGraphSchema.h"
+
+UGameplayAbilityGraphSchema::UGameplayAbilityGraphSchema(const FPostConstructInitializeProperties& PCIP)
+: Super(PCIP)
+{
+}
\ No newline at end of file
diff --git a/Engine/Source/Editor/GameplayAbilitiesEditor/Public/GameplayAbilitiesEditor.h b/Engine/Source/Editor/GameplayAbilitiesEditor/Public/GameplayAbilitiesEditor.h
new file mode 100644
--- /dev/null
+++ b/Engine/Source/Editor/GameplayAbilitiesEditor/Public/GameplayAbilitiesEditor.h
@@ -0,0 +1,47 @@
+// Copyright 1998-2014 Epic Games, Inc. All Rights Reserved.
+
+#pragma once
+
+#include "Editor/Kismet/Public/BlueprintEditor.h"
+
+#include "GameplayAbilitiesEditorModule.h"
+
+//////////////////////////////////////////////////////////////////////////
+// FGameplayAbilitiesEditor
+
+/**
+ * Gameplay abilities asset editor (extends Blueprint editor)
+ */
+class FGameplayAbilitiesEditor : public FBlueprintEditor
+{
+public:
+	/**
+	 * Edits the specified gameplay ability asset(s)
+	 *
+	 * @param	Mode					Asset editing mode for this editor (standalone or world-centric)
+	 * @param	InitToolkitHost			When Mode is WorldCentric, this is the level editor instance to spawn this editor within
+	 * @param	InBlueprints			The blueprints to edit
+	 * @param	bShouldOpenInDefaultsMode	If true, the editor will open in defaults editing mode
+	 */ 
+
+	void InitGameplayAbilitiesEditor(const EToolkitMode::Type Mode, const TSharedPtr<IToolkitHost>& InitToolkitHost, const TArray<UBlueprint*>& InBlueprints, bool bShouldOpenInDefaultsMode);
+public:
+	FGameplayAbilitiesEditor();
+
+	virtual ~FGameplayAbilitiesEditor();
+
+public:
+	// IToolkit interface
+	virtual FName GetToolkitFName() const override;
+	virtual FText GetBaseToolkitName() const override;
+	virtual FText GetToolkitName() const override;
+	virtual FString GetWorldCentricTabPrefix() const override;
+	virtual FLinearColor GetWorldCentricTabColorScale() const override;
+	// End of IToolkit interface
+
+	/** @return the documentation location for this editor */
+	virtual FString GetDocumentationLink() const override;
+	
+	/** Returns a pointer to the Blueprint object we are currently editing, as long as we are editing exactly one */
+	virtual UBlueprint* GetBlueprintObj() const override;
+};
diff --git a/Engine/Source/Runtime/GameplayAbilities/Classes/GameplayAbilityBlueprint.h b/Engine/Source/Runtime/GameplayAbilities/Classes/GameplayAbilityBlueprint.h
new file mode 100644
--- /dev/null
+++ b/Engine/Source/Runtime/GameplayAbilities/Classes/GameplayAbilityBlueprint.h
@@ -0,0 +1,33 @@
+// Copyright 1998-2014 Epic Games, Inc. All Rights Reserved.
+
+#pragma once
+
+#include "Engine/Blueprint.h"
+#include "GameplayAbilityBlueprint.generated.h"
+
+/**
+ * A Gameplay Ability Blueprint is essentially a specialized Blueprint whose graphs control a gameplay ability.
+ */
+
+UCLASS(dependson = (UBlueprint), BlueprintType)
+class GAMEPLAYABILITIES_API UGameplayAbilityBlueprint : public UBlueprint
+{
+	GENERATED_UCLASS_BODY()
+
+#if WITH_EDITOR
+
+	virtual UClass* GetBlueprintClass() const override;
+
+	// UBlueprint interface
+	virtual bool SupportedByDefaultBlueprintFactory() const override
+	{
+		return false;
+	}
+	// End of UBlueprint interface
+
+	/** Returns the most base gameplay ability blueprint for a given blueprint (if it is inherited from another ability blueprint, returning null if only native / non-ability BP classes are it's parent) */
+	static UGameplayAbilityBlueprint* FindRootGameplayAbilityBlueprint(UGameplayAbilityBlueprint* DerivedBlueprint);
+
+	virtual void PostLoad() override;
+#endif
+};
diff --git a/Engine/Source/Runtime/GameplayAbilities/Private/GameplayAbilityBlueprint.cpp b/Engine/Source/Runtime/GameplayAbilities/Private/GameplayAbilityBlueprint.cpp
new file mode 100644
--- /dev/null
+++ b/Engine/Source/Runtime/GameplayAbilities/Private/GameplayAbilityBlueprint.cpp
@@ -0,0 +1,45 @@
+// Copyright 1998-2014 Epic Games, Inc. All Rights Reserved.
+
+#include "AbilitySystemPrivatePCH.h"
+#include "GameplayAbilityBlueprint.h"
+#include "GameplayAbilityBlueprintGeneratedClass.h"
+
+//////////////////////////////////////////////////////////////////////////
+// UGameplayAbilityBlueprint
+
+UGameplayAbilityBlueprint::UGameplayAbilityBlueprint(const class FPostConstructInitializeProperties& PCIP)
+	: Super(PCIP)
+{
+}
+
+#if WITH_EDITOR
+
+UClass* UGameplayAbilityBlueprint::GetBlueprintClass() const
+{
+	return UGameplayAbilityBlueprint::StaticClass();
+}
+
+
+/** Returns the most base gameplay ability blueprint for a given blueprint (if it is inherited from another ability blueprint, returning null if only native / non-ability BP classes are it's parent) */
+UGameplayAbilityBlueprint* UGameplayAbilityBlueprint::FindRootGameplayAbilityBlueprint(UGameplayAbilityBlueprint* DerivedBlueprint)
+{
+	UGameplayAbilityBlueprint* ParentBP = NULL;
+
+	// Determine if there is a gameplay ability blueprint in the ancestry of this class
+	for (UClass* ParentClass = DerivedBlueprint->ParentClass; ParentClass != UObject::StaticClass(); ParentClass = ParentClass->GetSuperClass())
+	{
+		if (UGameplayAbilityBlueprint* TestBP = Cast<UGameplayAbilityBlueprint>(ParentClass->ClassGeneratedBy))
+		{
+			ParentBP = TestBP;
+		}
+	}
+
+	return ParentBP;
+}
+
+void UGameplayAbilityBlueprint::PostLoad()
+{
+	Super::PostLoad();
+}
+
+#endif
\ No newline at end of file
diff --git a/Engine/Source/Runtime/GameplayAbilities/Private/GameplayAbilityBlueprintGeneratedClass.cpp b/Engine/Source/Runtime/GameplayAbilities/Private/GameplayAbilityBlueprintGeneratedClass.cpp
new file mode 100644
--- /dev/null
+++ b/Engine/Source/Runtime/GameplayAbilities/Private/GameplayAbilityBlueprintGeneratedClass.cpp
@@ -0,0 +1,20 @@
+// Copyright 1998-2014 Epic Games, Inc. All Rights Reserved.
+
+#include "AbilitySystemPrivatePCH.h"
+
+#define LOCTEXT_NAMESPACE "GameplayAbility"
+
+/////////////////////////////////////////////////////
+// UGameplayAbilityBlueprintGeneratedClass
+
+UGameplayAbilityBlueprintGeneratedClass::UGameplayAbilityBlueprintGeneratedClass(const FPostConstructInitializeProperties& PCIP)
+	: Super(PCIP)
+{
+}
+
+void UGameplayAbilityBlueprintGeneratedClass::Link(FArchive& Ar, bool bRelinkExistingProperties)
+{
+	Super::Link(Ar, bRelinkExistingProperties);
+}
+
+#undef LOCTEXT_NAMESPACE
\ No newline at end of file
diff --git a/Engine/Source/Runtime/GameplayAbilities/Public/GameplayAbilityBlueprintGeneratedClass.h b/Engine/Source/Runtime/GameplayAbilities/Public/GameplayAbilityBlueprintGeneratedClass.h
new file mode 100644
--- /dev/null
+++ b/Engine/Source/Runtime/GameplayAbilities/Public/GameplayAbilityBlueprintGeneratedClass.h
@@ -0,0 +1,13 @@
+// Copyright 1998-2014 Epic Games, Inc. All Rights Reserved.
+
+#pragma once
+
+#include "GameplayAbilityBlueprintGeneratedClass.generated.h"
+
+UCLASS()
+class  GAMEPLAYABILITIES_API UGameplayAbilityBlueprintGeneratedClass : public UBlueprintGeneratedClass
+{
+	GENERATED_UCLASS_BODY()
+
+	virtual void Link(FArchive& Ar, bool bRelinkExistingProperties) override;
+};
\ No newline at end of file
