Fix non-unity compile errors created by CL# 2106979

diff --git a/Engine/Source/Editor/GameplayAbilitiesEditor/Private/GameplayAbilitiesBlueprintFactory.cpp b/Engine/Source/Editor/GameplayAbilitiesEditor/Private/GameplayAbilitiesBlueprintFactory.cpp
--- a/Engine/Source/Editor/GameplayAbilitiesEditor/Private/GameplayAbilitiesBlueprintFactory.cpp
+++ b/Engine/Source/Editor/GameplayAbilitiesEditor/Private/GameplayAbilitiesBlueprintFactory.cpp
@@ -6,6 +6,7 @@
 #include "Kismet2/KismetEditorUtilities.h"
 
 #include "Abilities/GameplayAbility.h"
+#include "GameplayAbilityBlueprint.h"
 
 #include "GameplayAbilityBlueprintGeneratedClass.h"
 
diff --git a/Engine/Source/Editor/GameplayAbilitiesEditor/Private/GameplayAbilitiesEditor.cpp b/Engine/Source/Editor/GameplayAbilitiesEditor/Private/GameplayAbilitiesEditor.cpp
--- a/Engine/Source/Editor/GameplayAbilitiesEditor/Private/GameplayAbilitiesEditor.cpp
+++ b/Engine/Source/Editor/GameplayAbilitiesEditor/Private/GameplayAbilitiesEditor.cpp
@@ -3,6 +3,8 @@
 #include "AbilitySystemEditorPrivatePCH.h"
 #include "GameplayAbilitiesEditor.h"
 
+#include "GameplayAbilityBlueprint.h"
+
 #define LOCTEXT_NAMESPACE "FGameplayAbilitiesEditor"
 
 
diff --git a/Engine/Source/Editor/GameplayAbilitiesEditor/Private/GameplayAbilitiesEditorModule.cpp b/Engine/Source/Editor/GameplayAbilitiesEditor/Private/GameplayAbilitiesEditorModule.cpp
--- a/Engine/Source/Editor/GameplayAbilitiesEditor/Private/GameplayAbilitiesEditorModule.cpp
+++ b/Engine/Source/Editor/GameplayAbilitiesEditor/Private/GameplayAbilitiesEditorModule.cpp
@@ -5,6 +5,10 @@
 #include "AttributeDetails.h"
 #include "AttributeSet.h"
 
+#include "IAssetTypeActions.h"
+#include "AssetToolsModule.h"
+#include "AssetTypeActions_GameplayAbilitiesBlueprint.h"
+
 class FGameplayAbilitiesEditorModule : public IGameplayAbilitiesEditorModule
 {
 	// Begin IModuleInterface
diff --git a/Engine/Source/Runtime/GameplayAbilities/Private/GameplayAbilityBlueprintGeneratedClass.cpp b/Engine/Source/Runtime/GameplayAbilities/Private/GameplayAbilityBlueprintGeneratedClass.cpp
--- a/Engine/Source/Runtime/GameplayAbilities/Private/GameplayAbilityBlueprintGeneratedClass.cpp
+++ b/Engine/Source/Runtime/GameplayAbilities/Private/GameplayAbilityBlueprintGeneratedClass.cpp
@@ -1,6 +1,7 @@
 // Copyright 1998-2014 Epic Games, Inc. All Rights Reserved.
 
 #include "AbilitySystemPrivatePCH.h"
+#include "GameplayAbilityBlueprintGeneratedClass.h"
 
 #define LOCTEXT_NAMESPACE "GameplayAbility"
 
