Added a parent class selection menu for GameplayAbilityBlueprints.

When creating a Gameplay Ability Blueprint the user will select the parent class from a filtered list. All classes on the list must be descended from UGameplayAbility.

diff --git a/Engine/Source/Editor/GameplayAbilitiesEditor/Private/GameplayAbilitiesBlueprintFactory.cpp b/Engine/Source/Editor/GameplayAbilitiesEditor/Private/GameplayAbilitiesBlueprintFactory.cpp
--- a/Engine/Source/Editor/GameplayAbilitiesEditor/Private/GameplayAbilitiesBlueprintFactory.cpp
+++ b/Engine/Source/Editor/GameplayAbilitiesEditor/Private/GameplayAbilitiesBlueprintFactory.cpp
@@ -8,12 +8,226 @@
 #include "Abilities/GameplayAbility.h"
 #include "GameplayAbilityBlueprint.h"
 
-#include "GameplayAbilityBlueprintGeneratedClass.h"
+#include "Editor/ClassViewer/Public/ClassViewerFilter.h"
 
+#include "GameplayAbilityBlueprintGeneratedClass.h"
 #include "GameplayAbilitiesBlueprintFactory.h"
 
 #define LOCTEXT_NAMESPACE "UGameplayAbilitiesBlueprintFactory"
 
+/*------------------------------------------------------------------------------
+Dialog to configure creation properties
+------------------------------------------------------------------------------*/
+class SGameplayAbilityBlueprintCreateDialog : public SCompoundWidget
+{
+public:
+	SLATE_BEGIN_ARGS(SGameplayAbilityBlueprintCreateDialog){}
+
+	SLATE_END_ARGS()
+
+		/** Constructs this widget with InArgs */
+		void Construct(const FArguments& InArgs)
+	{
+			bOkClicked = false;
+			ParentClass = UGameplayAbility::StaticClass();
+
+			ChildSlot
+				[
+					SNew(SBorder)
+					.Visibility(EVisibility::Visible)
+					.BorderImage(FEditorStyle::GetBrush("Menu.Background"))
+					[
+						SNew(SBox)
+						.Visibility(EVisibility::Visible)
+						.WidthOverride(500.0f)
+						[
+							SNew(SVerticalBox)
+							+ SVerticalBox::Slot()
+							.FillHeight(1)
+							[
+								SNew(SBorder)
+								.BorderImage(FEditorStyle::GetBrush("ToolPanel.GroupBorder"))
+								.Content()
+								[
+									SAssignNew(ParentClassContainer, SVerticalBox)
+								]
+							]
+
+							// Ok/Cancel buttons
+							+ SVerticalBox::Slot()
+								.AutoHeight()
+								.HAlign(HAlign_Right)
+								.VAlign(VAlign_Bottom)
+								.Padding(8)
+								[
+									SNew(SUniformGridPanel)
+									.SlotPadding(FEditorStyle::GetMargin("StandardDialog.SlotPadding"))
+									.MinDesiredSlotWidth(FEditorStyle::GetFloat("StandardDialog.MinDesiredSlotWidth"))
+									.MinDesiredSlotHeight(FEditorStyle::GetFloat("StandardDialog.MinDesiredSlotHeight"))
+									+ SUniformGridPanel::Slot(0, 0)
+									[
+										SNew(SButton)
+										.HAlign(HAlign_Center)
+										.ContentPadding(FEditorStyle::GetMargin("StandardDialog.ContentPadding"))
+										.OnClicked(this, &SGameplayAbilityBlueprintCreateDialog::OkClicked)
+										.Text(LOCTEXT("CreateGameplayAbilityBlueprintOk", "OK"))
+									]
+									+ SUniformGridPanel::Slot(1, 0)
+										[
+											SNew(SButton)
+											.HAlign(HAlign_Center)
+											.ContentPadding(FEditorStyle::GetMargin("StandardDialog.ContentPadding"))
+											.OnClicked(this, &SGameplayAbilityBlueprintCreateDialog::CancelClicked)
+											.Text(LOCTEXT("CreateGameplayAbilityBlueprintCancel", "Cancel"))
+										]
+								]
+						]
+					]
+				];
+
+			MakeParentClassPicker();
+		}
+
+	/** Sets properties for the supplied GameplayAbilitiesBlueprintFactory */
+	bool ConfigureProperties(TWeakObjectPtr<UGameplayAbilitiesBlueprintFactory> InGameplayAbilitiesBlueprintFactory)
+	{
+		GameplayAbilitiesBlueprintFactory = InGameplayAbilitiesBlueprintFactory;
+
+		TSharedRef<SWindow> Window = SNew(SWindow)
+			.Title(LOCTEXT("CreateGameplayAbilityBlueprintOptions", "Create Gameplay Ability Blueprint"))
+			.ClientSize(FVector2D(400, 700))
+			.SupportsMinimize(false).SupportsMaximize(false)
+			[
+				AsShared()
+			];
+
+		PickerWindow = Window;
+
+		GEditor->EditorAddModalWindow(Window);
+		GameplayAbilitiesBlueprintFactory.Reset();
+
+		return bOkClicked;
+	}
+
+private:
+	class FGameplayAbilityBlueprintParentFilter : public IClassViewerFilter
+	{
+	public:
+		/** All children of these classes will be included unless filtered out by another setting. */
+		TSet< const UClass* > AllowedChildrenOfClasses;
+
+		FGameplayAbilityBlueprintParentFilter() {}
+
+		virtual bool IsClassAllowed(const FClassViewerInitializationOptions& InInitOptions, const UClass* InClass, TSharedRef< FClassViewerFilterFuncs > InFilterFuncs) override
+		{
+			// If it appears on the allowed child-of classes list (or there is nothing on that list)
+			return InFilterFuncs->IfInChildOfClassesSet(AllowedChildrenOfClasses, InClass) != EFilterReturn::Failed;
+		}
+
+		virtual bool IsUnloadedClassAllowed(const FClassViewerInitializationOptions& InInitOptions, const TSharedRef< const IUnloadedBlueprintData > InUnloadedClassData, TSharedRef< FClassViewerFilterFuncs > InFilterFuncs) override
+		{
+			// If it appears on the allowed child-of classes list (or there is nothing on that list)
+			return InFilterFuncs->IfInChildOfClassesSet(AllowedChildrenOfClasses, InUnloadedClassData) != EFilterReturn::Failed;
+		}
+	};
+
+	/** Creates the combo menu for the parent class */
+	void MakeParentClassPicker()
+	{
+		// Load the classviewer module to display a class picker
+		FClassViewerModule& ClassViewerModule = FModuleManager::LoadModuleChecked<FClassViewerModule>("ClassViewer");
+
+		// Fill in options
+		FClassViewerInitializationOptions Options;
+		Options.Mode = EClassViewerMode::ClassPicker;
+
+		// Only allow parenting to base blueprints.
+		Options.bIsBlueprintBaseOnly = true;
+
+		TSharedPtr<FGameplayAbilityBlueprintParentFilter> Filter = MakeShareable(new FGameplayAbilityBlueprintParentFilter);
+
+		// All child child classes of UGameplayAbility are valid.
+		Filter->AllowedChildrenOfClasses.Add(UGameplayAbility::StaticClass());
+		Options.ClassFilter = Filter;
+
+		ParentClassContainer->ClearChildren();
+		ParentClassContainer->AddSlot()
+			.AutoHeight()
+			[
+				SNew(STextBlock)
+				.Text(LOCTEXT("ParentClass", "Parent Class:"))
+				.ShadowOffset(FVector2D(1.0f, 1.0f))
+			];
+
+		ParentClassContainer->AddSlot()
+			[
+				ClassViewerModule.CreateClassViewer(Options, FOnClassPicked::CreateSP(this, &SGameplayAbilityBlueprintCreateDialog::OnClassPicked))
+			];
+	}
+
+	/** Handler for when a parent class is selected */
+	void OnClassPicked(UClass* ChosenClass)
+	{
+		ParentClass = ChosenClass;
+	}
+
+	/** Handler for when ok is clicked */
+	FReply OkClicked()
+	{
+		if (GameplayAbilitiesBlueprintFactory.IsValid())
+		{
+			GameplayAbilitiesBlueprintFactory->BlueprintType = BPTYPE_Normal;
+			GameplayAbilitiesBlueprintFactory->ParentClass = ParentClass.Get();
+		}
+
+		CloseDialog(true);
+
+		return FReply::Handled();
+	}
+
+	void CloseDialog(bool bWasPicked = false)
+	{
+		bOkClicked = bWasPicked;
+		if (PickerWindow.IsValid())
+		{
+			PickerWindow.Pin()->RequestDestroyWindow();
+		}
+	}
+
+	/** Handler for when cancel is clicked */
+	FReply CancelClicked()
+	{
+		CloseDialog();
+		return FReply::Handled();
+	}
+
+	FReply OnKeyDown(const FGeometry& MyGeometry, const FKeyboardEvent& InKeyboardEvent)
+	{
+		if (InKeyboardEvent.GetKey() == EKeys::Escape)
+		{
+			CloseDialog();
+			return FReply::Handled();
+		}
+		return SWidget::OnKeyDown(MyGeometry, InKeyboardEvent);
+	}
+
+private:
+	/** The factory for which we are setting up properties */
+	TWeakObjectPtr<UGameplayAbilitiesBlueprintFactory> GameplayAbilitiesBlueprintFactory;
+
+	/** A pointer to the window that is asking the user to select a parent class */
+	TWeakPtr<SWindow> PickerWindow;
+
+	/** The container for the Parent Class picker */
+	TSharedPtr<SVerticalBox> ParentClassContainer;
+
+	/** The selected class */
+	TWeakObjectPtr<UClass> ParentClass;
+
+	/** True if Ok was clicked */
+	bool bOkClicked;
+};
+
 /*------------------------------------------------------------------------------
 	UGameplayAbilitiesBlueprintFactory implementation.
 ------------------------------------------------------------------------------*/
@@ -29,7 +243,8 @@ UGameplayAbilitiesBlueprintFactory::UGameplayAbilitiesBlueprintFactory(const cla
 
 bool UGameplayAbilitiesBlueprintFactory::ConfigureProperties()
 {
-	return true;
+	TSharedRef<SGameplayAbilityBlueprintCreateDialog> Dialog = SNew(SGameplayAbilityBlueprintCreateDialog);
+	return Dialog->ConfigureProperties(this);
 };
 
 UObject* UGameplayAbilitiesBlueprintFactory::FactoryCreateNew(UClass* Class, UObject* InParent, FName Name, EObjectFlags Flags, UObject* Context, FFeedbackContext* Warn, FName CallingContext)
