GameplayAbilities updates

diff --git a/Engine/Source/Editor/GameplayAbilitiesEditor/Classes/K2Node_LatentAbilityCall.h b/Engine/Source/Editor/GameplayAbilitiesEditor/Classes/K2Node_LatentAbilityCall.h
--- a/Engine/Source/Editor/GameplayAbilitiesEditor/Classes/K2Node_LatentAbilityCall.h
+++ b/Engine/Source/Editor/GameplayAbilitiesEditor/Classes/K2Node_LatentAbilityCall.h
@@ -3,9 +3,11 @@
 #pragma once
 #include "K2Node_LatentAbilityCall.generated.h"
 
-
 // TODO: Make this only available in GameplayAbility graphs!
 
+
+
+
 UCLASS()
 class UK2Node_LatentAbilityCall : public UK2Node_BaseAsyncTask
 {
@@ -14,4 +16,19 @@ class UK2Node_LatentAbilityCall : public UK2Node_BaseAsyncTask
 	// UEdGraphNode interface
 	virtual void GetMenuEntries(FGraphContextMenuBuilder& ContextMenuBuilder) const override;
 	// End of UEdGraphNode interface
+
+	virtual void AllocateDefaultPins() override;
+	virtual void ReallocatePinsDuringReconstruction(TArray<UEdGraphPin*>& OldPins) override;
+	virtual void PinDefaultValueChanged(UEdGraphPin* Pin) override;
+	virtual void ExpandNode(class FKismetCompilerContext& CompilerContext, UEdGraph* SourceGraph) override;
+
+	void CreatePinsForClass(UClass* InClass);
+	UEdGraphPin* GetClassPin(const TArray<UEdGraphPin*>* InPinsToSearch = NULL) const;
+	UClass* GetClassToSpawn(const TArray<UEdGraphPin*>* InPinsToSearch = NULL) const;
+	UEdGraphPin* GetResultPin() const;
+	bool IsSpawnVarPin(UEdGraphPin* Pin);
+	bool ValidateActorSpawning(class FKismetCompilerContext& CompilerContext);
+
+	UPROPERTY()
+	TArray<UEdGraphPin*>	SpawnParmPins;
 };
diff --git a/Engine/Source/Editor/GameplayAbilitiesEditor/Private/K2Node_GameplayCueEvent.cpp b/Engine/Source/Editor/GameplayAbilitiesEditor/Private/K2Node_GameplayCueEvent.cpp
--- a/Engine/Source/Editor/GameplayAbilitiesEditor/Private/K2Node_GameplayCueEvent.cpp
+++ b/Engine/Source/Editor/GameplayAbilitiesEditor/Private/K2Node_GameplayCueEvent.cpp
@@ -44,6 +44,8 @@ void UK2Node_GameplayCueEvent::GetMenuEntries(FGraphContextMenuBuilder& Context)
 
 	FGameplayTagContainer CueTags = GameplayTagsModule.GetGameplayTagsManager().RequestGameplayTagChildren(RootTag);
 
+	// Fixme: need to check if this function is already defined so that it can be reimplemented
+
 	for (auto It = CueTags.CreateConstIterator(); It; ++It)
 	{
 		FGameplayTag Tag = *It;
diff --git a/Engine/Source/Editor/GameplayAbilitiesEditor/Private/K2Node_LatentAbilityCall.cpp b/Engine/Source/Editor/GameplayAbilitiesEditor/Private/K2Node_LatentAbilityCall.cpp
--- a/Engine/Source/Editor/GameplayAbilitiesEditor/Private/K2Node_LatentAbilityCall.cpp
+++ b/Engine/Source/Editor/GameplayAbilitiesEditor/Private/K2Node_LatentAbilityCall.cpp
@@ -3,6 +3,8 @@
 
 #include "AbilitySystemEditorPrivatePCH.h"
 #include "AbilityTask.h"
+#include "KismetCompiler.h"
+#include "BlueprintEditorUtils.h"
 #include "K2Node_LatentAbilityCall.h"
 
 #define LOCTEXT_NAMESPACE "K2Node"
@@ -50,6 +52,465 @@ void UK2Node_LatentAbilityCall::GetMenuEntries(FGraphContextMenuBuilder& Context
 	}
 }
 
-#undef LOCTEXT_NAMESPACE
+// -------------------------------------------------
+
+struct FK2Node_LatentAbilityCallHelper
+{
+	static FString WorldContextPinName;
+	static FString ClassPinName;
+	static FString BeginSpawnFuncName;
+	static FString FinishSpawnFuncName;
+	static FString SpawnedActorPinName;
+};
+
+FString FK2Node_LatentAbilityCallHelper::WorldContextPinName(TEXT("WorldContextObject"));
+FString FK2Node_LatentAbilityCallHelper::ClassPinName(TEXT("Class"));
+FString FK2Node_LatentAbilityCallHelper::BeginSpawnFuncName(TEXT("BeginSpawningActor"));
+FString FK2Node_LatentAbilityCallHelper::FinishSpawnFuncName(TEXT("FinishSpawningActor"));
+FString FK2Node_LatentAbilityCallHelper::SpawnedActorPinName(TEXT("SpawnedActor"));
+
+// -------------------------------------------------
+
+void UK2Node_LatentAbilityCall::AllocateDefaultPins()
+{
+	Super::AllocateDefaultPins();
+}
+
+void UK2Node_LatentAbilityCall::ReallocatePinsDuringReconstruction(TArray<UEdGraphPin*>& OldPins)
+{
+	AllocateDefaultPins();
+	UClass* UseSpawnClass = GetClassToSpawn(&OldPins);
+	if (UseSpawnClass != NULL)
+	{
+		CreatePinsForClass(UseSpawnClass);
+	}
+}
+
+UEdGraphPin* UK2Node_LatentAbilityCall::GetClassPin(const TArray<UEdGraphPin*>* InPinsToSearch /*= NULL*/) const
+{
+	const TArray<UEdGraphPin*>* PinsToSearch = InPinsToSearch ? InPinsToSearch : &Pins;
+
+	UEdGraphPin* Pin = NULL;
+	for (auto PinIt = PinsToSearch->CreateConstIterator(); PinIt; ++PinIt)
+	{
+		UEdGraphPin* TestPin = *PinIt;
+		if (TestPin && TestPin->PinName == FK2Node_LatentAbilityCallHelper::ClassPinName)
+		{
+			Pin = TestPin;
+			break;
+		}
+	}
+	check(Pin == NULL || Pin->Direction == EGPD_Input);
+	return Pin;
+}
+
+UClass* UK2Node_LatentAbilityCall::GetClassToSpawn(const TArray<UEdGraphPin*>* InPinsToSearch) const
+{
+	UClass* UseSpawnClass = NULL;
+	const TArray<UEdGraphPin*>* PinsToSearch = InPinsToSearch ? InPinsToSearch : &Pins;
+
+	UEdGraphPin* ClassPin = GetClassPin(PinsToSearch);
+	if (ClassPin && ClassPin->DefaultObject != NULL && ClassPin->LinkedTo.Num() == 0)
+	{
+		UseSpawnClass = CastChecked<UClass>(ClassPin->DefaultObject);
+	}
+
+	return UseSpawnClass;
+}
+
+void UK2Node_LatentAbilityCall::CreatePinsForClass(UClass* InClass)
+{
+	check(InClass != NULL);
+
+	const UEdGraphSchema_K2* K2Schema = GetDefault<UEdGraphSchema_K2>();
+
+	SpawnParmPins.Empty();
+
+	for (TFieldIterator<UProperty> PropertyIt(InClass, EFieldIteratorFlags::IncludeSuper); PropertyIt; ++PropertyIt)
+	{
+		UProperty* Property = *PropertyIt;
+		UClass* PropertyClass = CastChecked<UClass>(Property->GetOuter());
+		const bool bIsDelegate = Property->IsA(UMulticastDelegateProperty::StaticClass());
+		const bool bIsExposedToSpawn = UEdGraphSchema_K2::IsPropertyExposedOnSpawn(Property);
+		const bool bIsSettableExternally = !Property->HasAnyPropertyFlags(CPF_DisableEditOnInstance);
+
+		if (bIsExposedToSpawn &&
+			!Property->HasAnyPropertyFlags(CPF_Parm) &&
+			bIsSettableExternally &&
+			Property->HasAllPropertyFlags(CPF_BlueprintVisible) &&
+			!bIsDelegate)
+		{
+			UEdGraphPin* Pin = CreatePin(EGPD_Input, TEXT(""), TEXT(""), NULL, false, false, Property->GetName());
+			const bool bPinGood = (Pin != NULL) && K2Schema->ConvertPropertyToPinType(Property, /*out*/ Pin->PinType);
+			Pin->bDefaultValueIsIgnored = true;
+			SpawnParmPins.Add(Pin);
+		}
+	}
+}
+
+void UK2Node_LatentAbilityCall::PinDefaultValueChanged(UEdGraphPin* ChangedPin)
+{
+	if (ChangedPin->PinName == FK2Node_LatentAbilityCallHelper::ClassPinName)
+	{
+		const UEdGraphSchema_K2* K2Schema = GetDefault<UEdGraphSchema_K2>();
+
+		// Because the archetype has changed, we break the output link as the output pin type will change
+		//UEdGraphPin* ResultPin = GetResultPin();
+		//ResultPin->BreakAllPinLinks();
 
+		// Remove all pins related to archetype variables
+		for (auto OldPin : SpawnParmPins)
+		{
+			OldPin->BreakAllPinLinks();
+			Pins.Remove(OldPin);
+		}
+		SpawnParmPins.Empty();
+
+		UClass* UseSpawnClass = GetClassToSpawn();
+		if (UseSpawnClass != NULL)
+		{
+			CreatePinsForClass(UseSpawnClass);
+		}
+
+		// Refresh the UI for the graph so the pin changes show up
+		UEdGraph* Graph = GetGraph();
+		Graph->NotifyGraphChanged();
+
+		// Mark dirty
+		FBlueprintEditorUtils::MarkBlueprintAsModified(GetBlueprint());
+	}
+}
+
+UEdGraphPin* UK2Node_LatentAbilityCall::GetResultPin() const
+{
+	const UEdGraphSchema_K2* K2Schema = GetDefault<UEdGraphSchema_K2>();
+
+	UEdGraphPin* Pin = FindPinChecked(K2Schema->PN_ReturnValue);
+	check(Pin->Direction == EGPD_Output);
+	return Pin;
+}
+
+bool UK2Node_LatentAbilityCall::IsSpawnVarPin(UEdGraphPin* Pin)
+{
+	const UEdGraphSchema_K2* K2Schema = GetDefault<UEdGraphSchema_K2>();
+
+	return (Pin->Direction == EEdGraphPinDirection::EGPD_Input &&
+			Pin->PinName != K2Schema->PN_Execute &&
+			Pin->PinName != K2Schema->PN_Then &&
+			Pin->PinName != K2Schema->PN_ReturnValue &&
+			Pin->PinName != FK2Node_LatentAbilityCallHelper::ClassPinName &&
+			Pin->PinName != FK2Node_LatentAbilityCallHelper::WorldContextPinName);
+
+}
+
+bool UK2Node_LatentAbilityCall::ValidateActorSpawning(class FKismetCompilerContext& CompilerContext)
+{
+	FName ProxyPrespawnFunctionName = *FK2Node_LatentAbilityCallHelper::BeginSpawnFuncName;
+	UFunction* PreSpawnFunction = ProxyFactoryClass->FindFunctionByName(ProxyPrespawnFunctionName);
 
+	FName ProxyPostpawnFunctionName = *FK2Node_LatentAbilityCallHelper::FinishSpawnFuncName;
+	UFunction* PostSpawnFunction = ProxyFactoryClass->FindFunctionByName(ProxyPostpawnFunctionName);
+
+	bool HasClassParameter = GetClassToSpawn() != nullptr;
+	bool HasPreSpawnFunc = PreSpawnFunction != nullptr;
+	bool HasPostSpawnFunc = PostSpawnFunction != nullptr;
+
+	
+	if (HasClassParameter || HasPreSpawnFunc || HasPostSpawnFunc)
+	{
+		// They are trying to use ActorSpawning. If any of the above are NOT true, then we have a problem
+		if (!HasClassParameter)
+		{
+			CompilerContext.MessageLog.Error(*LOCTEXT("MissingClassParameter", "UK2Node_LatentAbilityCall: Attempting to use ActorSpawning but Proxy Factory Function missing a Class parameter. @@").ToString(), this);
+			return false;
+		}
+		if (!HasPreSpawnFunc)
+		{
+			CompilerContext.MessageLog.Error(*LOCTEXT("MissingBeginSpawningFunc", "UK2Node_LatentAbilityCall: Attempting to use ActorSpawning but Missing a BeginSpawningActor function. @@").ToString(), this);
+			return false;
+		}
+		if (!HasPostSpawnFunc)
+		{
+			CompilerContext.MessageLog.Error(*LOCTEXT("MissingFinishSpawningFunc", "UK2Node_LatentAbilityCall: Attempting to use ActorSpawning but Missing a FinishSpawningActor function. @@").ToString(), this);
+			return false;
+		}
+	}	
+	
+	return true;
+	
+}
+
+/**
+ *	This is essentially a mix of K2Node_BaseAsyncTask::ExpandNode and K2Node_SpawnActorFromClass::ExpandNode.
+ *	Several things are going on here:
+ *		-Factory call to create proxy object (K2Node_BaseAsyncTask)
+ *		-Task return delegates are created and hooked up (K2Node_BaseAsyncTask)
+ *		-A BeginSpawn function is called on proxyu object (similiar to K2Node_SpawnActorFromClass)
+ *		-BeginSpawn can choose to spawn or not spawn an actor (and return it)
+ *			-If spawned:
+ *				-SetVars are run on the newly spawned object (set expose on spawn variables - K2Node_SpawnActorFromClass)
+ *				-FinishSpawn is called on the proxy object
+ *				
+ *				
+ *	Also, a K2Node_SpawnActorFromClass could not be used directly here, since we want the proxy object to implement its own
+ *	BeginSpawn/FinishSpawn function (custom game logic will often be performed in the native implementation). K2Node_SpawnActorFromClass also
+ *	requires a SpawnTransform be wired into it, and in most ability task cases, the spawn transform is implied or not necessary.
+ *	
+ *	
+ */
+void UK2Node_LatentAbilityCall::ExpandNode(class FKismetCompilerContext& CompilerContext, UEdGraph* SourceGraph)
+{
+	ValidateActorSpawning(CompilerContext);
+
+	UClass* ClassToSpawn = GetClassToSpawn();
+	if (ClassToSpawn == nullptr)
+	{
+		// Nothing special about this task, just call super
+		Super::ExpandNode(CompilerContext, SourceGraph);
+		return;
+	}
+
+	UK2Node::ExpandNode(CompilerContext, SourceGraph);
+
+	if (!CompilerContext.bIsFullCompile)
+	{
+		return;
+	}
+
+	const UEdGraphSchema_K2* Schema = CompilerContext.GetSchema();
+	check(SourceGraph && Schema);
+	bool bIsErrorFree = true;
+
+	// ------------------------------------------------------------------------------------------
+	// CREATE A CALL TO FACTORY THE PROXY OBJECT
+	// ------------------------------------------------------------------------------------------
+	UK2Node_CallFunction* const CallCreateProxyObjectNode = CompilerContext.SpawnIntermediateNode<UK2Node_CallFunction>(this, SourceGraph);
+	CallCreateProxyObjectNode->FunctionReference.SetExternalMember(ProxyFactoryFunctionName, ProxyFactoryClass);
+	CallCreateProxyObjectNode->AllocateDefaultPins();
+	bIsErrorFree &= CompilerContext.MovePinLinksToIntermediate(*FindPinChecked(Schema->PN_Execute), *CallCreateProxyObjectNode->FindPinChecked(Schema->PN_Execute)).CanSafeConnect();
+	for (auto CurrentPin : Pins)
+	{
+		if (FBaseAsyncTaskHelper::ValidDataPin(CurrentPin, EGPD_Input, Schema))
+		{
+			UEdGraphPin* DestPin = CallCreateProxyObjectNode->FindPin(CurrentPin->PinName); // match function inputs, to pass data to function from CallFunction node
+
+			// NEW: if no DestPin, assume it is a Class Spawn PRoperty - not an error
+			if (DestPin)
+			{
+				bIsErrorFree &= DestPin && CompilerContext.MovePinLinksToIntermediate(*CurrentPin, *DestPin).CanSafeConnect();
+			}
+		}
+	}
+	// ------------------------------------------------------------------------------------------
+	// GATHER OUTPUT PARAMETERS AND PAIR THEM WITH LOCAL VARIABLES
+	// ------------------------------------------------------------------------------------------
+	TArray<FBaseAsyncTaskHelper::FOutputPinAndLocalVariable> VariableOutputs;
+	for (auto CurrentPin : Pins)
+	{
+		if (FBaseAsyncTaskHelper::ValidDataPin(CurrentPin, EGPD_Output, Schema))
+		{
+			const FEdGraphPinType& PinType = CurrentPin->PinType;
+			UK2Node_TemporaryVariable* TempVarOutput = CompilerContext.SpawnInternalVariable(
+				this, PinType.PinCategory, PinType.PinSubCategory, PinType.PinSubCategoryObject.Get(), PinType.bIsArray);
+			bIsErrorFree &= TempVarOutput->GetVariablePin() && CompilerContext.MovePinLinksToIntermediate(*CurrentPin, *TempVarOutput->GetVariablePin()).CanSafeConnect();
+			VariableOutputs.Add(FBaseAsyncTaskHelper::FOutputPinAndLocalVariable(CurrentPin, TempVarOutput));
+		}
+	}
+
+	// ------------------------------------------------------------------------------------------
+	// FOR EACH DELEGATE DEFINE EVENT, CONNECT IT TO DELEGATE AND IMPLEMENT A CHAIN OF ASSIGMENTS
+	// ------------------------------------------------------------------------------------------
+	UEdGraphPin* LastThenPin = CallCreateProxyObjectNode->FindPinChecked(Schema->PN_Then);
+	UEdGraphPin* const ProxyObjectPin = CallCreateProxyObjectNode->GetReturnValuePin();
+	for (TFieldIterator<UMulticastDelegateProperty> PropertyIt(ProxyClass, EFieldIteratorFlags::ExcludeSuper); PropertyIt && bIsErrorFree; ++PropertyIt)
+	{
+		bIsErrorFree &= FBaseAsyncTaskHelper::HandleDelegateImplementation(*PropertyIt, VariableOutputs, ProxyObjectPin, LastThenPin, this, SourceGraph, CompilerContext);
+	}
+
+	if (CallCreateProxyObjectNode->FindPinChecked(Schema->PN_Then) == LastThenPin)
+	{
+		CompilerContext.MessageLog.Error(*LOCTEXT("MissingDelegateProperties", "BaseAsyncTask: Proxy has no delegates defined. @@").ToString(), this);
+		return;
+	}
+
+
+	// ------------------------------------------------------------------------------------------
+	// NEW: CREATE A CALL TO THE PRESPAWN FUNCTION, IF IT RETURNS TRUE, THEN WE WILL SPAWN THE NEW ACTOR
+	// ------------------------------------------------------------------------------------------
+	
+	FName ProxyPrespawnFunctionName = *FK2Node_LatentAbilityCallHelper::BeginSpawnFuncName;
+	UFunction* PreSpawnFunction = ProxyFactoryClass->FindFunctionByName(ProxyPrespawnFunctionName);
+
+	FName ProxyPostpawnFunctionName = *FK2Node_LatentAbilityCallHelper::FinishSpawnFuncName;
+	UFunction* PostSpawnFunction = ProxyFactoryClass->FindFunctionByName(ProxyPostpawnFunctionName);
+
+	if (PreSpawnFunction == nullptr)
+	{
+		CompilerContext.MessageLog.Error(*LOCTEXT("MissingBeginSpawningActorFunction", "AbilityTask: Proxy is missing BeginSpawningActor native function. @@").ToString(), this);
+		return;
+	}
+
+	if (PostSpawnFunction == nullptr)
+	{
+		CompilerContext.MessageLog.Error(*LOCTEXT("MissingFinishSpawningActorFunction", "AbilityTask: Proxy is missing FinishSpawningActor native function. @@").ToString(), this);
+		return;
+	}
+
+
+	UK2Node_CallFunction* const CallPrespawnProxyObjectNode = CompilerContext.SpawnIntermediateNode<UK2Node_CallFunction>(this, SourceGraph);
+	CallPrespawnProxyObjectNode->FunctionReference.SetExternalMember(ProxyPrespawnFunctionName, ProxyClass);
+	CallPrespawnProxyObjectNode->AllocateDefaultPins();
+
+	// Hook up the self connection
+	UEdGraphPin* PrespawnCallSelfPin = Schema->FindSelfPin(*CallPrespawnProxyObjectNode, EGPD_Input);
+	check(PrespawnCallSelfPin);
+
+	bIsErrorFree &= Schema->TryCreateConnection(ProxyObjectPin, PrespawnCallSelfPin);
+
+	// Hook up input parameters to PreSpawn
+	for (auto CurrentPin : Pins)
+	{
+		if (FBaseAsyncTaskHelper::ValidDataPin(CurrentPin, EGPD_Input, Schema))
+		{
+			UEdGraphPin* DestPin = CallPrespawnProxyObjectNode->FindPin(CurrentPin->PinName);
+			if (DestPin)
+			{
+				bIsErrorFree &= DestPin && CompilerContext.MovePinLinksToIntermediate(*CurrentPin, *DestPin).CanSafeConnect();
+			}
+		}
+	}		
+
+	// Hook the activate node up in the exec chain
+	UEdGraphPin* PrespawnExecPin = CallPrespawnProxyObjectNode->FindPinChecked(Schema->PN_Execute);
+	UEdGraphPin* PrespawnThenPin = CallPrespawnProxyObjectNode->FindPinChecked(Schema->PN_Then);
+	UEdGraphPin* PrespawnReturnPin = CallPrespawnProxyObjectNode->FindPinChecked(Schema->PN_ReturnValue);
+	UEdGraphPin* SpawnedActorReturnPin = CallPrespawnProxyObjectNode->FindPinChecked(FK2Node_LatentAbilityCallHelper::SpawnedActorPinName);
+
+	bIsErrorFree &= Schema->TryCreateConnection(LastThenPin, PrespawnExecPin);
+
+	LastThenPin = PrespawnThenPin;
+
+	// -------------------------------------------
+	// Branch based on return value of Prespawn
+	// -------------------------------------------
+		
+	UK2Node_IfThenElse* BranchNode = SourceGraph->CreateBlankNode<UK2Node_IfThenElse>();
+	BranchNode->AllocateDefaultPins();
+	CompilerContext.MessageLog.NotifyIntermediateObjectCreation(BranchNode, this);
+
+	// Link return value of prespawn with the branch condtional
+	bIsErrorFree &= Schema->TryCreateConnection(PrespawnReturnPin, BranchNode->GetConditionPin());
+
+	// Link our Prespawn call to the branch node
+	bIsErrorFree &= Schema->TryCreateConnection(LastThenPin, BranchNode->GetExecPin());
+
+	UEdGraphPin* BranchElsePin = BranchNode->GetElsePin();
+
+	LastThenPin = BranchNode->GetThenPin();
+
+	// -------------------------------------------
+	// Set spawn variables
+	// -------------------------------------------
+
+	for (auto SpawnVarPin : SpawnParmPins)
+	{
+		if (SpawnVarPin->LinkedTo.Num() > 0)
+		{
+			UFunction* SetByNameFunction = Schema->FindSetVariableByNameFunction(SpawnVarPin->PinType);
+			if (SetByNameFunction)
+			{
+				UK2Node_CallFunction* SetVarNode = NULL;
+				if (SpawnVarPin->PinType.bIsArray)
+				{
+					SetVarNode = CompilerContext.SpawnIntermediateNode<UK2Node_CallArrayFunction>(this, SourceGraph);
+				}
+				else
+				{
+					SetVarNode = CompilerContext.SpawnIntermediateNode<UK2Node_CallFunction>(this, SourceGraph);
+				}
+				SetVarNode->SetFromFunction(SetByNameFunction);
+				SetVarNode->AllocateDefaultPins();
+
+				// Connect this node into the exec chain
+				bIsErrorFree &= Schema->TryCreateConnection(LastThenPin, SetVarNode->GetExecPin());
+				LastThenPin = SetVarNode->GetThenPin();
+
+				static FString ObjectParamName = FString(TEXT("Object"));
+				static FString ValueParamName = FString(TEXT("Value"));
+				static FString PropertyNameParamName = FString(TEXT("PropertyName"));
+
+				// Connect the new actor to the 'object' pin
+				UEdGraphPin* ObjectPin = SetVarNode->FindPinChecked(ObjectParamName);
+				SpawnedActorReturnPin->MakeLinkTo(ObjectPin);
+
+				// Fill in literal for 'property name' pin - name of pin is property name
+				UEdGraphPin* PropertyNamePin = SetVarNode->FindPinChecked(PropertyNameParamName);
+				PropertyNamePin->DefaultValue = SpawnVarPin->PinName;
+
+				// Move connection from the variable pin on the spawn node to the 'value' pin
+				UEdGraphPin* ValuePin = SetVarNode->FindPinChecked(ValueParamName);
+				CompilerContext.MovePinLinksToIntermediate(*SpawnVarPin, *ValuePin);
+				SetVarNode->PinConnectionListChanged(ValuePin);
+			}
+		}
+	}
+		
+
+	// -------------------------------------------
+	// Call FinishSpawning
+	// -------------------------------------------
+
+	UK2Node_CallFunction* const CallPostSpawnnProxyObjectNode = CompilerContext.SpawnIntermediateNode<UK2Node_CallFunction>(this, SourceGraph);
+	CallPostSpawnnProxyObjectNode->FunctionReference.SetExternalMember(ProxyPostpawnFunctionName, ProxyClass);
+	CallPostSpawnnProxyObjectNode->AllocateDefaultPins();
+
+	// Hook up the self connection
+	UEdGraphPin* PostspawnCallSelfPin = Schema->FindSelfPin(*CallPostSpawnnProxyObjectNode, EGPD_Input);
+	check(PostspawnCallSelfPin);
+
+	bIsErrorFree &= Schema->TryCreateConnection(ProxyObjectPin, PostspawnCallSelfPin);
+
+	// Link our Postspawn call in
+	bIsErrorFree &= Schema->TryCreateConnection(LastThenPin, CallPostSpawnnProxyObjectNode->FindPinChecked(Schema->PN_Execute));
+
+	// Hook up any other input parameters to PostSpawn
+	for (auto CurrentPin : Pins)
+	{
+		if (FBaseAsyncTaskHelper::ValidDataPin(CurrentPin, EGPD_Input, Schema))
+		{
+			UEdGraphPin* DestPin = CallPostSpawnnProxyObjectNode->FindPin(CurrentPin->PinName);
+			if (DestPin)
+			{
+				bIsErrorFree &= DestPin && CompilerContext.MovePinLinksToIntermediate(*CurrentPin, *DestPin).CanSafeConnect();
+			}
+		}
+	}
+
+
+	UEdGraphPin* InSpawnedActorPin = CallPostSpawnnProxyObjectNode->FindPin(TEXT("SpawnedActor"));
+	if (InSpawnedActorPin == nullptr)
+	{
+		CompilerContext.MessageLog.Error(*LOCTEXT("MissingSpawnedActorInputPin", "AbilityTask: Proxy is missing SpawnedActor input pin in FinishSpawningActor. @@").ToString(), this);
+		return;
+	}
+
+	bIsErrorFree &= Schema->TryCreateConnection(SpawnedActorReturnPin, InSpawnedActorPin);
+
+	LastThenPin = CallPostSpawnnProxyObjectNode->FindPinChecked(Schema->PN_Then);
+
+
+	// Move the connections from the original node then pin to the last internal then pin
+	bIsErrorFree &= CompilerContext.MovePinLinksToIntermediate(*FindPinChecked(Schema->PN_Then), *LastThenPin).CanSafeConnect();
+	bIsErrorFree &= CompilerContext.CopyPinLinksToIntermediate(*LastThenPin, *BranchElsePin).CanSafeConnect();
+	
+	if (!bIsErrorFree)
+	{
+		CompilerContext.MessageLog.Error(*LOCTEXT("InternalConnectionError", "BaseAsyncTask: Internal connection error. @@").ToString(), this);
+	}
+
+	// Make sure we caught everything
+	BreakAllNodeLinks();
+}
+
+
+#undef LOCTEXT_NAMESPACE
diff --git a/Engine/Source/Runtime/GameplayAbilities/Private/Abilities/GameplayAbility.cpp b/Engine/Source/Runtime/GameplayAbilities/Private/Abilities/GameplayAbility.cpp
--- a/Engine/Source/Runtime/GameplayAbilities/Private/Abilities/GameplayAbility.cpp
+++ b/Engine/Source/Runtime/GameplayAbilities/Private/Abilities/GameplayAbility.cpp
@@ -400,6 +400,13 @@ FGameplayAbilityActorInfo UGameplayAbility::GetActorInfo()
 	return *CurrentActorInfo;
 }
 
+/** Convenience method for abilities to get outgoing gameplay effect specs (for example, to pass on to projectiles to apply to whoever they hit) */
+FGameplayEffectSpecHandle UGameplayAbility::GetOutgoingSpec(UGameplayEffect* GameplayEffect) const
+{
+	check(CurrentActorInfo && CurrentActorInfo->AbilitySystemComponent.IsValid());
+	return CurrentActorInfo->AbilitySystemComponent->GetOutgoingSpec(GameplayEffect);
+}
+
 /** Fixme: Naming is confusing here */
 
 bool UGameplayAbility::K2_CommitAbility()
diff --git a/Engine/Source/Runtime/GameplayAbilities/Private/Abilities/GameplayAbilityTargetActor_SingleLineTrace.cpp b/Engine/Source/Runtime/GameplayAbilities/Private/Abilities/GameplayAbilityTargetActor_SingleLineTrace.cpp
--- a/Engine/Source/Runtime/GameplayAbilities/Private/Abilities/GameplayAbilityTargetActor_SingleLineTrace.cpp
+++ b/Engine/Source/Runtime/GameplayAbilities/Private/Abilities/GameplayAbilityTargetActor_SingleLineTrace.cpp
@@ -18,7 +18,7 @@ AGameplayAbilityTargetActor_SingleLineTrace::AGameplayAbilityTargetActor_SingleL
 	PrimaryActorTick.TickGroup = TG_PrePhysics;
 	StaticTargetFunction = false;
 	bDebug = false;
-	bBindToConfirmCancelInputs = false;
+	bBindToConfirmCancelInputs = true;
 }
 
 FGameplayAbilityTargetDataHandle AGameplayAbilityTargetActor_SingleLineTrace::StaticGetTargetData(UWorld * World, const FGameplayAbilityActorInfo* ActorInfo, FGameplayAbilityActivationInfo ActivationInfo)
diff --git a/Engine/Source/Runtime/GameplayAbilities/Private/Abilities/Tasks/AbilityTask_SpawnActor.cpp b/Engine/Source/Runtime/GameplayAbilities/Private/Abilities/Tasks/AbilityTask_SpawnActor.cpp
new file mode 100644
--- /dev/null
+++ b/Engine/Source/Runtime/GameplayAbilities/Private/Abilities/Tasks/AbilityTask_SpawnActor.cpp
@@ -0,0 +1,52 @@
+
+#include "AbilitySystemPrivatePCH.h"
+#include "GameplayAbilityTargetActor.h"
+#include "AbilitySystemComponent.h"
+#include "Abilities/Tasks/AbilityTask_SpawnActor.h"
+
+#pragma optimize( "", off )
+
+UAbilityTask_SpawnActor::UAbilityTask_SpawnActor(const class FPostConstructInitializeProperties& PCIP)
+	: Super(PCIP)
+{
+	
+}
+
+UAbilityTask_SpawnActor* UAbilityTask_SpawnActor::SpawnActor(UObject* WorldContextObject, TSubclassOf<AActor> InClass)
+{
+	UAbilityTask_SpawnActor * MyObj = NewObject<UAbilityTask_SpawnActor>();
+	UGameplayAbility * ThisAbility = CastChecked<UGameplayAbility>(WorldContextObject);
+	MyObj->InitTask(ThisAbility);
+	return MyObj;	
+}
+
+// ---------------------------------------------------------------------------------------
+
+bool UAbilityTask_SpawnActor::BeginSpawningActor(UObject* WorldContextObject, TSubclassOf<AActor> InClass, AActor*& SpawnedActor)
+{
+	if (Ability.IsValid() && Ability.Get()->GetCurrentActorInfo()->IsNetAuthority())
+	{
+		UWorld* World = GEngine->GetWorldFromContextObject(WorldContextObject);
+		SpawnedActor = World->SpawnActorDeferred<AActor>(InClass, FVector::ZeroVector, FRotator::ZeroRotator, NULL, NULL, true);
+	}
+	
+	return (SpawnedActor != nullptr);
+}
+
+void UAbilityTask_SpawnActor::FinishSpawningActor(UObject* WorldContextObject, AActor* SpawnedActor)
+{
+	if (SpawnedActor)
+	{
+		FTransform SpawnTransform = AbilitySystemComponent->GetOwner()->GetTransform();
+
+		SpawnedActor->ExecuteConstruction(SpawnTransform, NULL);
+		SpawnedActor->PostActorConstruction();
+
+		Success.Broadcast(SpawnedActor);
+	}
+
+	MarkPendingKill();
+}
+
+// ---------------------------------------------------------------------------------------
+
diff --git a/Engine/Source/Runtime/GameplayAbilities/Private/Abilities/Tasks/AbilityTask_WaitTargetData.cpp b/Engine/Source/Runtime/GameplayAbilities/Private/Abilities/Tasks/AbilityTask_WaitTargetData.cpp
--- a/Engine/Source/Runtime/GameplayAbilities/Private/Abilities/Tasks/AbilityTask_WaitTargetData.cpp
+++ b/Engine/Source/Runtime/GameplayAbilities/Private/Abilities/Tasks/AbilityTask_WaitTargetData.cpp
@@ -19,29 +19,15 @@ UAbilityTask_WaitTargetData* UAbilityTask_WaitTargetData::WaitTargetData(UObject
 	return MyObj;	
 }
 
-void UAbilityTask_WaitTargetData::Activate()
+// ---------------------------------------------------------------------------------------
+
+bool UAbilityTask_WaitTargetData::BeginSpawningActor(UObject* WorldContextObject, TSubclassOf<AGameplayAbilityTargetActor> TargetClass, AGameplayAbilityTargetActor*& SpawnedActor)
 {
+	SpawnedActor = nullptr;
+
 	if (Ability.IsValid())
 	{
-		if (!Ability.Get()->GetCurrentActorInfo()->IsLocallyControlled())
-		{
-			// If not locally controlled (server for remote client), see if TargetData was already sent
-			// else register callback for when it does get here
-
-			if (AbilitySystemComponent->ReplicatedTargetData.IsValid())
-			{
-				ValidData.Broadcast(AbilitySystemComponent->ReplicatedTargetData);
-				Cleanup();
-				return;
-			}
-			else
-			{
-				AbilitySystemComponent->ReplicatedTargetDataDelegate.AddUObject(this, &UAbilityTask_WaitTargetData::OnTargetDataReplicatedCallback);
-				AbilitySystemComponent->ReplicatedTargetDataCancelledDelegate.AddDynamic(this, &UAbilityTask_WaitTargetData::OnTargetDataReplicatedCancelledCallback);
-				return;
-			}
-		}
-		else
+		if (Ability.Get()->GetCurrentActorInfo()->IsLocallyControlled())
 		{
 			// Locally controlled - spawn the targeting actor.
 			AGameplayAbilityTargetActor* CDO = CastChecked<AGameplayAbilityTargetActor>(TargetClass->GetDefaultObject());
@@ -53,23 +39,60 @@ void UAbilityTask_WaitTargetData::Activate()
 			}
 			else
 			{
-				// This is a latent thing, probably with visuals and other gameplay related stuff. Spawn an actor locally and it will tell us when it has TargetData.
-				AGameplayAbilityTargetActor* SpawnedActor = CastChecked<AGameplayAbilityTargetActor>(Ability->GetWorld()->SpawnActor(TargetClass));
+				UClass* Class = *TargetClass;
+				if (Class != NULL)
+				{
+					UWorld* World = GEngine->GetWorldFromContextObject(WorldContextObject);
+					SpawnedActor = World->SpawnActorDeferred<AGameplayAbilityTargetActor>(Class, FVector::ZeroVector, FRotator::ZeroRotator, NULL, NULL, true);
+				}
+
 				SpawnedActor->TargetDataReadyDelegate.AddUObject(this, &UAbilityTask_WaitTargetData::OnTargetDataReadyCallback);
 				SpawnedActor->CanceledDelegate.AddUObject(this, &UAbilityTask_WaitTargetData::OnTargetDataCancelledCallback);
-								
-				// User ability activation is inhibited while this is active
-				AbilitySystemComponent->SetUserAbilityActivationInhibited(true);
-				AbilitySystemComponent->SpawnedTargetActors.Push(SpawnedActor);
-
-				SpawnedActor->StartTargeting(Ability.Get());
 
 				MySpawnedTargetActor = SpawnedActor;
 			}
 		}
+		else
+		{
+			// If not locally controlled (server for remote client), see if TargetData was already sent
+			// else register callback for when it does get here
+
+			if (AbilitySystemComponent->ReplicatedTargetData.IsValid())
+			{
+				ValidData.Broadcast(AbilitySystemComponent->ReplicatedTargetData);
+				Cleanup();
+			}
+			else
+			{
+				AbilitySystemComponent->ReplicatedTargetDataDelegate.AddUObject(this, &UAbilityTask_WaitTargetData::OnTargetDataReplicatedCallback);
+				AbilitySystemComponent->ReplicatedTargetDataCancelledDelegate.AddDynamic(this, &UAbilityTask_WaitTargetData::OnTargetDataReplicatedCancelledCallback);
+			}
+		}
 	}
+	return (SpawnedActor != nullptr);
 }
 
+void UAbilityTask_WaitTargetData::FinishSpawningActor(UObject* WorldContextObject, AGameplayAbilityTargetActor* SpawnedActor)
+{
+	if (SpawnedActor)
+	{
+		check(MySpawnedTargetActor == SpawnedActor);
+
+		FTransform SpawnTransform = AbilitySystemComponent->GetOwner()->GetTransform();
+
+		SpawnedActor->ExecuteConstruction(SpawnTransform, NULL);
+		SpawnedActor->PostActorConstruction();
+
+		// User ability activation is inhibited while this is active
+		AbilitySystemComponent->SetUserAbilityActivationInhibited(true);
+		AbilitySystemComponent->SpawnedTargetActors.Push(SpawnedActor);
+
+		SpawnedActor->StartTargeting(Ability.Get());
+	}
+}
+
+// ---------------------------------------------------------------------------------------
+
 /** Valid TargetData was replicated to use (we are server, was sent from client) */
 void UAbilityTask_WaitTargetData::OnTargetDataReplicatedCallback(FGameplayAbilityTargetDataHandle Data)
 {
@@ -126,4 +149,8 @@ void UAbilityTask_WaitTargetData::Cleanup()
 	}
 
 	MarkPendingKill();
-}
\ No newline at end of file
+}
+
+
+// --------------------------------------------------------------------------------------
+
diff --git a/Engine/Source/Runtime/GameplayAbilities/Private/AbilitySystemComponent.cpp b/Engine/Source/Runtime/GameplayAbilities/Private/AbilitySystemComponent.cpp
--- a/Engine/Source/Runtime/GameplayAbilities/Private/AbilitySystemComponent.cpp
+++ b/Engine/Source/Runtime/GameplayAbilities/Private/AbilitySystemComponent.cpp
@@ -173,12 +173,34 @@ float UAbilitySystemComponent::GetNumericAttribute(const FGameplayAttribute &Att
 	return Attribute.GetNumericValueChecked(AttributeSet);
 }
 
+FGameplayEffectSpecHandle UAbilitySystemComponent::GetOutgoingSpec(UGameplayEffect* GameplayEffect) const
+{
+	SCOPE_CYCLE_COUNTER(STAT_GetOutgoingSpec);
+	// Fixme: we should build a map and cache these off. We can invalidate the map when an OutgoingGE modifier is applied or removed from us.
+
+	FGameplayEffectSpec* NewSpec = new FGameplayEffectSpec(GameplayEffect, GetOwner(), FGameplayEffectLevelSpec::INVALID_LEVEL, GetCurveDataOverride());
+	if (ActiveGameplayEffects.ApplyActiveEffectsTo(*NewSpec, FModifierQualifier().Type(EGameplayMod::OutgoingGE)))
+	{
+		return FGameplayEffectSpecHandle(NewSpec);
+	}
+
+	delete NewSpec;
+	return FGameplayEffectSpecHandle(nullptr);
+}
+
+FGameplayEffectInstigatorContext UAbilitySystemComponent::GetInstigatorContext() const
+{
+	FGameplayEffectInstigatorContext Context;
+	Context.AddInstigator(GetOwner());
+	return Context;
+}
+
 /** This is a helper function used in automated testing, I'm not sure how useful it will be to gamecode or blueprints */
 FActiveGameplayEffectHandle UAbilitySystemComponent::ApplyGameplayEffectToTarget(UGameplayEffect *GameplayEffect, UAbilitySystemComponent *Target, float Level, FModifierQualifier BaseQualifier)
 {
 	check(GameplayEffect);
 
-	FGameplayEffectSpec	Spec(GameplayEffect, GetOwner(),  Level, GetCurveDataOverride());
+	FGameplayEffectSpec	Spec(GameplayEffect, GetOwner(), Level, GetCurveDataOverride());
 	
 	return ApplyGameplayEffectSpecToTarget(Spec, Target, BaseQualifier);
 }
diff --git a/Engine/Source/Runtime/GameplayAbilities/Private/AbilitySystemStats.cpp b/Engine/Source/Runtime/GameplayAbilities/Private/AbilitySystemStats.cpp
--- a/Engine/Source/Runtime/GameplayAbilities/Private/AbilitySystemStats.cpp
+++ b/Engine/Source/Runtime/GameplayAbilities/Private/AbilitySystemStats.cpp
@@ -11,3 +11,6 @@ DEFINE_STAT(STAT_GameplayAbilityGetCooldownTimeRemaining);
 DEFINE_STAT(STAT_RemoveActiveGameplayEffect);
 DEFINE_STAT(STAT_CreateNewActiveGameplayEffect);
 DEFINE_STAT(STAT_GetGameplayCueFunction);
+DEFINE_STAT(STAT_GetOutgoingSpec);
+
+
diff --git a/Engine/Source/Runtime/GameplayAbilities/Private/GameplayEffect.cpp b/Engine/Source/Runtime/GameplayAbilities/Private/GameplayEffect.cpp
--- a/Engine/Source/Runtime/GameplayAbilities/Private/GameplayEffect.cpp
+++ b/Engine/Source/Runtime/GameplayAbilities/Private/GameplayEffect.cpp
@@ -179,7 +179,7 @@ void FGameplayEffectSpec::MakeUnique()
 	}
 }
 
-int32 FGameplayEffectSpec::ApplyModifiersFrom(FGameplayEffectSpec &InSpec, const FModifierQualifier &QualifierContext)
+int32 FGameplayEffectSpec::ApplyModifiersFrom(const FGameplayEffectSpec &InSpec, const FModifierQualifier &QualifierContext)
 {
 	ABILITY_LOG_SCOPE(TEXT("FGameplayEffectSpec::ApplyModifiersFrom %s. InSpec: %s"), *this->ToSimpleString(), *InSpec.ToSimpleString());
 
@@ -1158,12 +1158,10 @@ void FActiveGameplayEffect::PostReplicatedAdd(const struct FActiveGameplayEffect
 //
 // --------------------------------------------------------------------------------------------------------------------------------------------------------
 
-bool FActiveGameplayEffectsContainer::ApplyActiveEffectsTo(OUT FGameplayEffectSpec &Spec, const FModifierQualifier &QualifierContext)
+bool FActiveGameplayEffectsContainer::ApplyActiveEffectsTo(OUT FGameplayEffectSpec &Spec, const FModifierQualifier &QualifierContext) const
 {
-	FActiveGameplayEffectHandle().IsValid();
-
 	ABILITY_LOG_SCOPE(TEXT("ApplyActiveEffectsTo: %s %s"), *Spec.ToSimpleString(), *QualifierContext.ToString());
-	for (FActiveGameplayEffect & ActiveEffect : GameplayEffects)
+	for (const FActiveGameplayEffect & ActiveEffect : GameplayEffects)
 	{
 		// We dont want to use FModifierQualifier::TestTarget here, since we aren't the 'target'. We are applying stuff to Spec which will be applied to a target.
 		if (QualifierContext.IgnoreHandle().IsValid() && QualifierContext.IgnoreHandle() == ActiveEffect.Handle)
diff --git a/Engine/Source/Runtime/GameplayAbilities/Public/Abilities/GameplayAbility.h b/Engine/Source/Runtime/GameplayAbilities/Public/Abilities/GameplayAbility.h
--- a/Engine/Source/Runtime/GameplayAbilities/Public/Abilities/GameplayAbility.h
+++ b/Engine/Source/Runtime/GameplayAbilities/Public/Abilities/GameplayAbility.h
@@ -6,6 +6,7 @@
 #include "GameplayTagAssetInterface.h"
 #include "Runtime/Engine/Classes/Animation/AnimInstance.h"
 #include "Abilities/GameplayAbilityTypes.h"
+#include "GameplayEffect.h"
 #include "GameplayAbility.generated.h"
 
 /**
@@ -301,9 +302,12 @@ protected:
 		}
 	}
 
-	UFUNCTION(BlueprintPure, Category = Ability)
+	UFUNCTION(BlueprintPure, Category=Ability)
 	FGameplayAbilityActorInfo GetActorInfo();
 
+	UFUNCTION(BlueprintCallable, Category=Ability)
+	FGameplayEffectSpecHandle GetOutgoingSpec(UGameplayEffect* GameplayEffect) const;
+
 	UFUNCTION(Client, Reliable)
 	void ClientActivateAbilitySucceed(int32 PredictionKey);
 
diff --git a/Engine/Source/Runtime/GameplayAbilities/Public/Abilities/GameplayAbilityTargetActor_SingleLineTrace.h b/Engine/Source/Runtime/GameplayAbilities/Public/Abilities/GameplayAbilityTargetActor_SingleLineTrace.h
--- a/Engine/Source/Runtime/GameplayAbilities/Public/Abilities/GameplayAbilityTargetActor_SingleLineTrace.h
+++ b/Engine/Source/Runtime/GameplayAbilities/Public/Abilities/GameplayAbilityTargetActor_SingleLineTrace.h
@@ -27,8 +27,12 @@ public:
 	virtual void Tick(float DeltaSeconds) override;
 	
 	TWeakObjectPtr<UGameplayAbility> Ability;
-
+	
 	bool bDebug;
 
+	UPROPERTY(BlueprintReadOnly, meta=(ExposeOnSpawn=true), Category=Projectile)
 	bool bBindToConfirmCancelInputs;
+
+	UPROPERTY(BlueprintReadOnly, meta = (ExposeOnSpawn = true), Category = Projectile)
+	FVector StartingOffset;
 };
\ No newline at end of file
diff --git a/Engine/Source/Runtime/GameplayAbilities/Public/Abilities/GameplayAbilityTypes.h b/Engine/Source/Runtime/GameplayAbilities/Public/Abilities/GameplayAbilityTypes.h
--- a/Engine/Source/Runtime/GameplayAbilities/Public/Abilities/GameplayAbilityTypes.h
+++ b/Engine/Source/Runtime/GameplayAbilities/Public/Abilities/GameplayAbilityTypes.h
@@ -275,6 +275,8 @@ struct GAMEPLAYABILITIES_API FGameplayAbilityTargetData
 		return TArray<AActor*>();
 	}
 
+	// -------------------------------------
+
 	virtual bool HasHitResult() const
 	{
 		return false;
@@ -285,6 +287,20 @@ struct GAMEPLAYABILITIES_API FGameplayAbilityTargetData
 		return NULL;
 	}
 
+	// -------------------------------------
+
+	virtual bool HasOrigin() const
+	{
+		return false;
+	}
+
+	virtual FTransform* GetOrigin() const
+	{
+		return nullptr;
+	}
+
+	// -------------------------------------
+
 	virtual UScriptStruct * GetScriptStruct()
 	{
 		return FGameplayAbilityTargetData::StaticStruct();
@@ -305,6 +321,8 @@ struct GAMEPLAYABILITIES_API FGameplayAbilityTargetData_SingleTargetHit : public
 	: HitResult(InHitResult)
 	{ }
 
+	// -------------------------------------
+
 	virtual TArray<AActor*>	GetActors() const
 	{
 		TArray<AActor*>	Actors;
@@ -314,6 +332,8 @@ struct GAMEPLAYABILITIES_API FGameplayAbilityTargetData_SingleTargetHit : public
 		}
 		return Actors;
 	}
+	
+	// -------------------------------------
 
 	virtual bool HasHitResult() const
 	{
@@ -325,6 +345,8 @@ struct GAMEPLAYABILITIES_API FGameplayAbilityTargetData_SingleTargetHit : public
 		return &HitResult;
 	}
 
+	// -------------------------------------
+
 	UPROPERTY()
 	FHitResult	HitResult;
 
diff --git a/Engine/Source/Runtime/GameplayAbilities/Public/Abilities/Tasks/AbilityTask.h b/Engine/Source/Runtime/GameplayAbilities/Public/Abilities/Tasks/AbilityTask.h
--- a/Engine/Source/Runtime/GameplayAbilities/Public/Abilities/Tasks/AbilityTask.h
+++ b/Engine/Source/Runtime/GameplayAbilities/Public/Abilities/Tasks/AbilityTask.h
@@ -3,8 +3,55 @@
 #include "GameplayAbility.h"
 #include "AbilityTask.generated.h"
 
-UCLASS(MinimalAPI)
-class UAbilityTask : public UObject
+
+/**
+ *	AbilityTasks are small, self contained operations that can be performed while executing an ability.
+ *	They are latent/asynchronous is nature. They will generate follow the pattern of 'start something and wait until it is finished or interrupted'
+ *	
+ *	We have code in K2Node_LatentAbilityCall to make using these in blueprints streamlined. The best way to become familiar with AbilityTasks is to 
+ *	look at existing tasks like UAbilityTask_WaitOverlap (very simple) and UAbilityTask_WaitTargetData (much more complex).
+ *	
+ *	These are the basic requirements for using an ability task:
+ *	
+ *	1) Define dynamic multicast, BlueprintAssignable delegates in your AbilityTask. These are the OUTPUTs of your task. When these delegates fire,
+ *	execution resumes in the calling blueprints.
+ *	
+ *	2) Your inputs are defined by a static factory function which will instantiate an instance of your task. The parameters of this function define
+ *	the INPUTs into your task. All the factory function should do is instantiate your task and possibly set starting parameters. It should NOT invoke
+ *	any of the callback delegates!
+ *	
+ *	3) Implement a Activate() function (defined here in base class). This function should actually start/execute your task logic. It is safe to invoke
+ *	callback delegates here.
+ *	
+ *	
+ *	This is all you need for basic AbilityTasks. 
+ *	
+ *	
+ *	
+ *	We have additional support for AbilityTasks that want to spawn actors. Though this could be accomplished in an Activate() function, it would not be
+ *	possible to pass in dynamic "ExposeOnSpawn" actor properties. This is a powerful feature of blueprints, in order to support this, you need to implement 
+ *	a different step 3:
+ *	
+ *	Instead of an Activate() function, you should implement a BeginSpawningActor() and FinishSpawningActor() function.
+ *	
+ *	BeginSpawningActor() must take in a TSubclassOf<YourActorClassToSpawn> parameters named 'Class'. It must also have a out reference parameters of type 
+ *	YourActorClassToSpawn*& named SpawnedActor. This function is allowed to decide whether it wants to spawn the actor or not (useful if wishing to
+ *	predicate actor spawning on network authority).
+ *	
+ *	BeginSpawningActor() can instantiate an actor with SpawnActorDefferred. This is important, otherwise the UCS will run before spawn parameters are set.
+ *	BeginSpawningActor() should also set the SpawnedActor parameter to the actor it spawned.
+ *	
+ *	[Next, the generated byte code will set the expose on spawn parameters to whatever the user has set]
+ *	
+ *	If you spawned something, FinishSpawningActor() will be called and pass in the same actor that was just spawned. You MUST call ExecuteConstruction + PostActorConstruction
+ *	on this actor!
+ *	
+ *	This is a lot of steps but in general, AbilityTask_SpawnActor() gives a clear, minimal example.
+ *	
+ *	
+ */
+UCLASS()
+class GAMEPLAYABILITIES_API UAbilityTask : public UObject
 {
 	GENERATED_UCLASS_BODY()
 
diff --git a/Engine/Source/Runtime/GameplayAbilities/Public/Abilities/Tasks/AbilityTask_SpawnActor.h b/Engine/Source/Runtime/GameplayAbilities/Public/Abilities/Tasks/AbilityTask_SpawnActor.h
new file mode 100644
--- /dev/null
+++ b/Engine/Source/Runtime/GameplayAbilities/Public/Abilities/Tasks/AbilityTask_SpawnActor.h
@@ -0,0 +1,42 @@
+// Copyright 1998-2014 Epic Games, Inc. All Rights Reserved.
+#pragma once
+
+#include "AbilityTask.h"
+#include "AbilityTask_SpawnActor.generated.h"
+
+class AGameplayAbilityTargetActor;
+
+/**
+ *	Convenience task for spawning actors on the network authority. If not the net authority, we will not spawn and Success will not be called.
+ *	The nice thing this adds is the ability to modify expose on spawn properties while also implicitly checking network role before spawning.
+ *	
+ *	Though this task doesn't do much - games can implement similiar tasks that carry out game specific rules. For example a 'SpawnProjectile'
+ *	task that limits the available classes to the games projectile class, and that does game specific stuff on spawn (for example, determining
+ *	firing position from a weapon attachment - logic that we don't necessarily want in ability blueprints).
+ *	
+ *	Long term we can also use this task as a sync point. If the executing client could wait execution until the server creates and replicate sthe 
+ *	actor down to him. We could potentially also use this to do predictive actor spawning / reconciliation.
+ */
+
+UCLASS(MinimalAPI)
+class UAbilityTask_SpawnActor: public UAbilityTask
+{
+	GENERATED_UCLASS_BODY()
+
+	DECLARE_DYNAMIC_MULTICAST_DELEGATE_OneParam(FSpawnActorDelegate, AActor*, SpawnedActor);
+
+	UPROPERTY(BlueprintAssignable)
+	FSpawnActorDelegate	Success;
+	
+	UFUNCTION(BlueprintCallable, meta=(HidePin = "WorldContextObject", DefaultToSelf = "WorldContextObject", BlueprintInternalUseOnly = "true"), Category = "Abilities")
+	static UAbilityTask_SpawnActor* SpawnActor(UObject* WorldContextObject, TSubclassOf<AActor> Class);
+
+	UFUNCTION(BlueprintCallable, meta = (HidePin = "WorldContextObject", DefaultToSelf = "WorldContextObject", BlueprintInternalUseOnly = "true"), Category = "Abilities")
+	bool BeginSpawningActor(UObject* WorldContextObject, TSubclassOf<AActor> Class, AActor*& SpawnedActor);
+
+	UFUNCTION(BlueprintCallable, meta = (HidePin = "WorldContextObject", DefaultToSelf = "WorldContextObject", BlueprintInternalUseOnly = "true"), Category = "Abilities")
+	void FinishSpawningActor(UObject* WorldContextObject, AActor* SpawnedActor);
+
+protected:
+	
+};
\ No newline at end of file
diff --git a/Engine/Source/Runtime/GameplayAbilities/Public/Abilities/Tasks/AbilityTask_WaitTargetData.h b/Engine/Source/Runtime/GameplayAbilities/Public/Abilities/Tasks/AbilityTask_WaitTargetData.h
--- a/Engine/Source/Runtime/GameplayAbilities/Public/Abilities/Tasks/AbilityTask_WaitTargetData.h
+++ b/Engine/Source/Runtime/GameplayAbilities/Public/Abilities/Tasks/AbilityTask_WaitTargetData.h
@@ -30,12 +30,27 @@ class UAbilityTask_WaitTargetData: public UAbilityTask
 
 	UFUNCTION()
 	void OnTargetDataCancelledCallback(FGameplayAbilityTargetDataHandle Data);
-	
+
+	/**
+	 *	Requirements for using Begin/Finish SpawningActor functionality:
+	 *		-Have a parameters named 'Class' in your Proxy factor function (E.g., WaitTargetdata)
+	 *		-Have a function named BeginSpawningActor w/ the same Class parameter
+	 *			-This function should spawn the actor with SpawnActorDeferred and return true/false if it spawned something.
+	 *		-Have a function named FinishSpawningActor w/ an AActor* of the class you spawned
+	 *			-This function *must* call ExecuteConstruction + PostActorConstruction
+	 */
 	
 	UFUNCTION(BlueprintCallable, meta=(HidePin = "WorldContextObject", DefaultToSelf = "WorldContextObject", BlueprintInternalUseOnly = "true"), Category = "Abilities")
-	static UAbilityTask_WaitTargetData* WaitTargetData(UObject* WorldContextObject, TSubclassOf<AGameplayAbilityTargetActor> TargetClass);
+	static UAbilityTask_WaitTargetData* WaitTargetData(UObject* WorldContextObject, TSubclassOf<AGameplayAbilityTargetActor> Class);
+
+	UFUNCTION(BlueprintCallable, meta = (HidePin = "WorldContextObject", DefaultToSelf = "WorldContextObject", BlueprintInternalUseOnly = "true"), Category = "Abilities")
+	bool BeginSpawningActor(UObject* WorldContextObject, TSubclassOf<AGameplayAbilityTargetActor> Class, AGameplayAbilityTargetActor*& SpawnedActor);
+
+	UFUNCTION(BlueprintCallable, meta = (HidePin = "WorldContextObject", DefaultToSelf = "WorldContextObject", BlueprintInternalUseOnly = "true"), Category = "Abilities")
+	void FinishSpawningActor(UObject* WorldContextObject, AGameplayAbilityTargetActor* SpawnedActor);
+
 
-	virtual void Activate() override;
+protected:
 
 	TSubclassOf<AGameplayAbilityTargetActor> TargetClass;
 
diff --git a/Engine/Source/Runtime/GameplayAbilities/Public/AbilitySystemComponent.h b/Engine/Source/Runtime/GameplayAbilities/Public/AbilitySystemComponent.h
--- a/Engine/Source/Runtime/GameplayAbilities/Public/AbilitySystemComponent.h
+++ b/Engine/Source/Runtime/GameplayAbilities/Public/AbilitySystemComponent.h
@@ -67,7 +67,7 @@ class GAMEPLAYABILITIES_API UAbilitySystemComponent : public UActorComponent
 
 	/** Finds existing AttributeSet */
 	template <class T >
-	T*	GetSet()
+	T*	GetSet()	// FIXME: This should be const by default
 	{
 		return (T*)GetAttributeSubobject(T::StaticClass());
 	}
@@ -172,12 +172,20 @@ class GAMEPLAYABILITIES_API UAbilitySystemComponent : public UActorComponent
 	// --------------------------------------------
 	// Primary outward facing API for other systems:
 	// --------------------------------------------
-	FActiveGameplayEffectHandle ApplyGameplayEffectSpecToTarget(OUT FGameplayEffectSpec &GameplayEffect, UAbilitySystemComponent *Target, FModifierQualifier BaseQualifier = FModifierQualifier());
-	FActiveGameplayEffectHandle ApplyGameplayEffectSpecToSelf(const FGameplayEffectSpec &GameplayEffect, FModifierQualifier BaseQualifier = FModifierQualifier());
+	FActiveGameplayEffectHandle ApplyGameplayEffectSpecToTarget(OUT FGameplayEffectSpec& GameplayEffect, UAbilitySystemComponent *Target, FModifierQualifier BaseQualifier = FModifierQualifier());
+	FActiveGameplayEffectHandle ApplyGameplayEffectSpecToSelf(const FGameplayEffectSpec& GameplayEffect, FModifierQualifier BaseQualifier = FModifierQualifier());
 
 	UFUNCTION(BlueprintCallable, BlueprintAuthorityOnly, Category = GameplayEffects)
 	bool RemoveActiveGameplayEffect(FActiveGameplayEffectHandle Handle);
 
+	/** Get an outgoing GameplayEffectSpec that is ready to be applied to other things. */
+	UFUNCTION(BlueprintCallable, Category = GameplayEffects)
+	FGameplayEffectSpecHandle GetOutgoingSpec(UGameplayEffect* GameplayEffect) const;
+
+	/** Create an InstigatorContext for the owner of this AbilitySystemComponent */
+	UFUNCTION(BlueprintCallable, Category = GameplayEffects)
+	FGameplayEffectInstigatorContext GetInstigatorContext() const;
+
 	/** This only exists so it can be hooked up to a multicast delegate */
 	void RemoveActiveGameplayEffect_NoReturn(FActiveGameplayEffectHandle Handle)
 	{
diff --git a/Engine/Source/Runtime/GameplayAbilities/Public/AbilitySystemStats.h b/Engine/Source/Runtime/GameplayAbilities/Public/AbilitySystemStats.h
--- a/Engine/Source/Runtime/GameplayAbilities/Public/AbilitySystemStats.h
+++ b/Engine/Source/Runtime/GameplayAbilities/Public/AbilitySystemStats.h
@@ -14,4 +14,5 @@ DECLARE_CYCLE_STAT_EXTERN(TEXT("GetActiveEffectsData"), STAT_GameplayEffectsGetA
 DECLARE_CYCLE_STAT_EXTERN(TEXT("GetCooldownTimeRemaining"), STAT_GameplayAbilityGetCooldownTimeRemaining, STATGROUP_AbilitySystem, );
 DECLARE_CYCLE_STAT_EXTERN(TEXT("RemoveActiveGameplayEffect"), STAT_RemoveActiveGameplayEffect, STATGROUP_AbilitySystem, );
 DECLARE_CYCLE_STAT_EXTERN(TEXT("CreateNewActiveGameplayEffect"), STAT_CreateNewActiveGameplayEffect, STATGROUP_AbilitySystem, );
-DECLARE_CYCLE_STAT_EXTERN(TEXT("GetGameplayCueFunction"), STAT_GetGameplayCueFunction, STATGROUP_AbilitySystem, );
\ No newline at end of file
+DECLARE_CYCLE_STAT_EXTERN(TEXT("GetGameplayCueFunction"), STAT_GetGameplayCueFunction, STATGROUP_AbilitySystem, );
+DECLARE_CYCLE_STAT_EXTERN(TEXT("GetOutgoingSpec"), STAT_GetOutgoingSpec, STATGROUP_AbilitySystem, );
\ No newline at end of file
diff --git a/Engine/Source/Runtime/GameplayAbilities/Public/GameplayEffect.h b/Engine/Source/Runtime/GameplayAbilities/Public/GameplayEffect.h
--- a/Engine/Source/Runtime/GameplayAbilities/Public/GameplayEffect.h
+++ b/Engine/Source/Runtime/GameplayAbilities/Public/GameplayEffect.h
@@ -2,16 +2,15 @@
 
 #pragma once
 
-#include "GameplayTagContainer.h"
 #include "GameplayTagAssetInterface.h"
+#include "AbilitySystemLog.h"
+#include "GameplayTagContainer.h"
 #include "TimerManager.h"
 #include "GameplayEffectTypes.h"
 #include "GameplayEffect.generated.h"
 
 struct FActiveGameplayEffect;
 
-
-
 class UGameplayEffect;
 class UAbilitySystemComponent;
 
@@ -966,9 +965,10 @@ struct FGameplayEffectSpec
 	float GetPeriod() const;
 	float GetChanceToApplyToTarget() const;
 	float GetChanceToExecuteOnGameplayEffect() const;
-	EGameplayEffectStackingPolicy::Type GetStackingType() const;
 	float GetMagnitude(const FGameplayAttribute &Attribute) const;
 
+	EGameplayEffectStackingPolicy::Type GetStackingType() const;
+
 	// other effects that need to be applied to the target if this effect is successful
 	TArray< TSharedRef< FGameplayEffectSpec > > TargetEffectSpecs;
 
@@ -985,13 +985,14 @@ struct FGameplayEffectSpec
 	FAggregatorRef	ChanceToExecuteOnGameplayEffect;
 
 	// This should only be true if this is a stacking effect and at the top of its stack
+	// (FIXME: should this be part of the spec or FActiveGameplayEffect?)
 	UPROPERTY()
 	bool bTopOfStack;
 
 	// The spec needs to own these FModifierSpecs so that other people can keep TSharedPtr to it.
 	// The stuff in this array is OWNED by this spec
 
-	TArray<FModifierSpec>	Modifiers;
+	TArray<FModifierSpec> Modifiers;
 
 	void MakeUnique();
 
@@ -999,7 +1000,7 @@ struct FGameplayEffectSpec
 
 	// returns the number of modifiers applied to InSpec by the current GameplayEffect Spec
 	// returns -1 if the current GameplayEffectSpec prevents InSpec from being applied
-	int32 ApplyModifiersFrom(FGameplayEffectSpec &InSpec, const FModifierQualifier &QualifierContext);
+	int32 ApplyModifiersFrom(const FGameplayEffectSpec &InSpec, const FModifierQualifier &QualifierContext);
 
 	int32 ExecuteModifiersFrom(const FGameplayEffectSpec &InSpec, const FModifierQualifier &QualifierContext);
 
@@ -1146,7 +1147,7 @@ struct FActiveGameplayEffectsContainer : public FFastArraySerializer
 
 	// returns true if none of the active effects provide immunity to Spec
 	// returns false if one (or more) of the active effects provides immunity to Spec
-	bool ApplyActiveEffectsTo(OUT FGameplayEffectSpec &Spec, const FModifierQualifier &QualifierContext);
+	bool ApplyActiveEffectsTo(OUT FGameplayEffectSpec &Spec, const FModifierQualifier &QualifierContext) const;
 
 	void ApplySpecToActiveEffectsAndAttributes(FGameplayEffectSpec &Spec, const FModifierQualifier &QualifierContext);
 		
@@ -1162,8 +1163,6 @@ struct FActiveGameplayEffectsContainer : public FFastArraySerializer
 
 	float GetGameplayEffectMagnitude(FActiveGameplayEffectHandle Handle, FGameplayAttribute Attribute) const;
 
-
-
 	// returns true if the handle points to an effect in this container that is not a stacking effect or an effect in this container that does stack and is applied by the current stacking rules
 	// returns false if the handle points to an effect that is not in this container or is not applied because of the current stacking rules
 	bool IsGameplayEffectActive(FActiveGameplayEffectHandle Handle) const;
@@ -1251,3 +1250,63 @@ struct TStructOpsTypeTraits< FActiveGameplayEffectsContainer > : public TStructO
 		WithNetDeltaSerializer = true,
 	};
 };
+
+
+USTRUCT(BlueprintType)
+struct GAMEPLAYABILITIES_API FGameplayEffectSpecHandle
+{
+	GENERATED_USTRUCT_BODY()
+
+	FGameplayEffectSpecHandle() { }
+	FGameplayEffectSpecHandle(FGameplayEffectSpec* DataPtr)
+		: Data(DataPtr)
+	{
+
+	}
+
+	TSharedPtr<FGameplayEffectSpec>	Data;
+
+	bool IsValidCache;
+
+	void Clear()
+	{
+		Data.Reset();
+	}
+
+	bool IsValid() const
+	{
+		return Data.IsValid();
+	}
+
+	bool NetSerialize(FArchive& Ar, class UPackageMap* Map, bool& bOutSuccess)
+	{
+		ABILITY_LOG(Fatal, TEXT("FGameplayEffectSpecHandle should not be NetSerialized"));
+		return false;
+	}
+
+	/** Comparison operator */
+	bool operator==(FGameplayEffectSpecHandle const& Other) const
+	{
+		// Both invalid structs or both valid and Pointer compare (???) // deep comparison equality
+		bool bBothValid = IsValid() && Other.IsValid();
+		bool bBothInvalid = !IsValid() && !Other.IsValid();
+		return (bBothInvalid || (bBothValid && (Data.Get() == Other.Data.Get())));
+	}
+
+	/** Comparison operator */
+	bool operator!=(FGameplayEffectSpecHandle const& Other) const
+	{
+		return !(FGameplayEffectSpecHandle::operator==(Other));
+	}
+};
+
+template<>
+struct TStructOpsTypeTraits<FGameplayEffectSpecHandle> : public TStructOpsTypeTraitsBase
+{
+	enum
+	{
+		WithCopy = true,		// Necessary so that TSharedPtr<FGameplayAbilityTargetData> Data is copied around
+		WithNetSerializer = true,
+		WithIdenticalViaEquality = true,
+	};
+};
\ No newline at end of file
diff --git a/Engine/Source/Runtime/GameplayAbilities/Public/GameplayEffectTypes.h b/Engine/Source/Runtime/GameplayAbilities/Public/GameplayEffectTypes.h
--- a/Engine/Source/Runtime/GameplayAbilities/Public/GameplayEffectTypes.h
+++ b/Engine/Source/Runtime/GameplayAbilities/Public/GameplayEffectTypes.h
@@ -2,6 +2,7 @@
 
 #pragma once
 
+#include "GameplayTagAssetInterface.h"
 #include "AttributeSet.h"
 #include "GameplayEffectTypes.generated.h"
 
@@ -197,7 +198,7 @@ struct TStructOpsTypeTraits< FGameplayEffectInstigatorContext > : public TStruct
 	enum
 	{
 		WithNetSerializer = true,
-		WithCopy = true		// Necessary so that TSharedPtr<FGameplayAbilityTargetData> Data is copied around
+		WithCopy = true		// Necessary so that TSharedPtr<FHitResult> Data is copied around
 	};
 };
 
