GameplayAbilities - only show ability tasks in GameplayAbilityBlueprints

diff --git a/Engine/Source/Editor/GameplayAbilitiesEditor/Classes/K2Node_LatentAbilityCall.h b/Engine/Source/Editor/GameplayAbilitiesEditor/Classes/K2Node_LatentAbilityCall.h
--- a/Engine/Source/Editor/GameplayAbilitiesEditor/Classes/K2Node_LatentAbilityCall.h
+++ b/Engine/Source/Editor/GameplayAbilitiesEditor/Classes/K2Node_LatentAbilityCall.h
@@ -14,6 +14,7 @@ class UK2Node_LatentAbilityCall : public UK2Node_BaseAsyncTask
 	GENERATED_UCLASS_BODY()
 
 	// UEdGraphNode interface
+	virtual bool CanCreateUnderSpecifiedSchema(const UEdGraphSchema* DesiredSchema) const;
 	virtual void GetMenuEntries(FGraphContextMenuBuilder& ContextMenuBuilder) const override;
 	// End of UEdGraphNode interface
 
diff --git a/Engine/Source/Editor/GameplayAbilitiesEditor/Private/K2Node_LatentAbilityCall.cpp b/Engine/Source/Editor/GameplayAbilitiesEditor/Private/K2Node_LatentAbilityCall.cpp
--- a/Engine/Source/Editor/GameplayAbilitiesEditor/Private/K2Node_LatentAbilityCall.cpp
+++ b/Engine/Source/Editor/GameplayAbilitiesEditor/Private/K2Node_LatentAbilityCall.cpp
@@ -6,6 +6,7 @@
 #include "KismetCompiler.h"
 #include "BlueprintEditorUtils.h"
 #include "K2Node_LatentAbilityCall.h"
+#include "GameplayAbilityGraphSchema.h"
 
 #define LOCTEXT_NAMESPACE "K2Node"
 
@@ -15,6 +16,11 @@ UK2Node_LatentAbilityCall::UK2Node_LatentAbilityCall(const class FPostConstructI
 	ProxyActivateFunctionName = GET_FUNCTION_NAME_CHECKED(UAbilityTask, Activate);
 }
 
+bool UK2Node_LatentAbilityCall::CanCreateUnderSpecifiedSchema(const UEdGraphSchema* DesiredSchema) const
+{
+	return DesiredSchema->GetClass()->IsChildOf(UGameplayAbilityGraphSchema::StaticClass());
+}
+
 void UK2Node_LatentAbilityCall::GetMenuEntries(FGraphContextMenuBuilder& ContextMenuBuilder) const
 {
 	const UEdGraphSchema_K2* K2Schema = GetDefault<UEdGraphSchema_K2>();
