Missed files from CL 2170702.

diff --git a/Engine/Source/Editor/GameplayAbilitiesEditor/Private/SGraphNodeK2GameplayEffectVar.cpp b/Engine/Source/Editor/GameplayAbilitiesEditor/Private/SGraphNodeK2GameplayEffectVar.cpp
new file mode 100644
--- /dev/null
+++ b/Engine/Source/Editor/GameplayAbilitiesEditor/Private/SGraphNodeK2GameplayEffectVar.cpp
@@ -0,0 +1,217 @@
+// Copyright 1998-2014 Epic Games, Inc. All Rights Reserved.
+
+
+#include "AbilitySystemEditorPrivatePCH.h"
+
+#include "SNodePanel.h"
+#include "SGraphNode.h"
+#include "../../Engine/Source/Editor/GraphEditor/Private/KismetNodes/SGraphNodeK2Base.h"
+#include "SGraphNodeK2GameplayEffectVar.h"
+#include "K2Node_GameplayEffectVariable.h"
+#include "GameplayEffectTypes.h"
+#include "ClassIconFinder.h"
+#include "IDocumentation.h"
+
+#define LOCTEXT_NAMESPACE "GraphNodeGameplayEffectVar"
+
+void SGraphNodeK2GameplayEffectVar::Construct( const FArguments& InArgs, UK2Node* InNode )
+{
+	this->GraphNode = InNode;
+
+	this->SetCursor( EMouseCursor::CardinalCross );
+
+	this->UpdateGraphNode();
+}
+
+FSlateColor SGraphNodeK2GameplayEffectVar::GetVariableColor() const
+{
+	return GraphNode->GetNodeTitleColor();
+}
+
+void SGraphNodeK2GameplayEffectVar::UpdateGraphNode()
+{
+	SGraphNodeK2Var::UpdateGraphNode();
+
+	UK2Node_GameplayEffectVariable* GameplayEffectNode = Cast<UK2Node_GameplayEffectVariable>(GraphNode);
+
+	FString DurationString = FString::Printf(TEXT("%d s"), GameplayEffectNode->GameplayEffectInfo.Duration);
+	FString PeriodString = FString::Printf(TEXT("%d s"), GameplayEffectNode->GameplayEffectInfo.Period);
+
+	RightNodeBox->AddSlot()
+		[
+			SNew(SHorizontalBox)
+			+ SHorizontalBox::Slot()
+			.HAlign(HAlign_Left)
+			[
+				SNew(STextBlock)
+				.Text(FString("Duration:"))
+			]
+			+ SHorizontalBox::Slot()
+			.HAlign(HAlign_Right)
+			[
+				SNew(STextBlock)
+				.Text(DurationString)
+			]
+		];
+
+	RightNodeBox->AddSlot()
+		[
+			SNew(SHorizontalBox)
+			+ SHorizontalBox::Slot()
+			.HAlign(HAlign_Left)
+			[
+				SNew(STextBlock)
+				.Text(FString("Period:"))
+			]
+			+ SHorizontalBox::Slot()
+				.HAlign(HAlign_Right)
+				[
+					SNew(STextBlock)
+					.Text(PeriodString)
+				]
+		];
+
+	RightNodeBox->AddSlot()
+		.FillHeight(0.3f)
+		[
+			SNew(SSpacer)
+		];
+
+	for (int ModIdx = 0; ModIdx < GameplayEffectNode->GameplayEffectInfo.AttributeModifiers.Num(); ++ModIdx)
+	{
+		RightNodeBox->AddSlot()
+			.FillHeight(0.3f)
+			[
+				SNew(SSpacer)
+			];
+
+		RightNodeBox->AddSlot()
+			.VAlign(VAlign_Center)
+			.HAlign(HAlign_Left)
+			[
+				SNew(STextBlock)
+				.Text(GameplayEffectNode->GameplayEffectInfo.AttributeModifiers[ModIdx].DisplayName)
+			];
+
+		FString ModOpString;
+		switch (GameplayEffectNode->GameplayEffectInfo.AttributeModifiers[ModIdx].ModOp)
+		{
+		case EGameplayModOp::Additive:
+			ModOpString = "Add";
+			break;
+		case EGameplayModOp::Multiplicitive:
+			ModOpString = "Multiply";
+			break;
+		case EGameplayModOp::Division:
+			ModOpString = "Divide";
+			break;
+		default:
+			ModOpString = "Other";
+			break;
+		}
+		FString ModString = FString::Printf(TEXT("%s: %f"), *ModOpString, GameplayEffectNode->GameplayEffectInfo.AttributeModifiers[ModIdx].Magnitude);
+
+		RightNodeBox->AddSlot()
+			.VAlign(VAlign_Center)
+			.HAlign(HAlign_Left)
+			[
+				SNew(STextBlock)
+				.Text(ModString)
+			];
+	}
+
+// 	TSharedRef<SWidget> OpenGameplayEffectButton = OpenGameplayEffectButtonContent(
+// 		LOCTEXT("FormatTextNodeOpenGameplayEffectButton", "Open Gameplay Effect"),
+// 		LOCTEXT("FormatTextNodeOpenGameplayEffectButton_Tooltip", "Opens this gameplay effect so that it can be edited."));
+// 
+// 	LeftNodeBox->AddSlot()
+// 		.AutoHeight()
+// 		.VAlign(VAlign_Center)
+// 		.Padding(10, 10, 10, 4)
+// 		[
+// 			OpenGameplayEffectButton
+// 		];
+}
+
+TSharedRef<SWidget> SGraphNodeK2GameplayEffectVar::OpenGameplayEffectButtonContent(FText ButtonText, FText ButtonTooltipText, FString DocumentationExcerpt)
+{
+	TSharedPtr<SWidget> ButtonContent;
+
+	SAssignNew(ButtonContent, SHorizontalBox)
+		+ SHorizontalBox::Slot()
+		.AutoWidth()
+		.VAlign(VAlign_Center)
+		.Padding(0, 0, 7, 0)
+		[
+			SNew(SImage)
+			.Image(FEditorStyle::GetBrush(TEXT("PropertyWindow.Button_Browse")))
+		]
+	+ SHorizontalBox::Slot()
+		.AutoWidth()
+		.HAlign(HAlign_Left)
+		[
+			SNew(STextBlock)
+			.Text(ButtonText)
+			.ColorAndOpacity(FLinearColor::White)
+		];
+
+
+	TSharedPtr<SToolTip> Tooltip;
+
+	if (!DocumentationExcerpt.IsEmpty())
+	{
+		Tooltip = IDocumentation::Get()->CreateToolTip(ButtonTooltipText, NULL, GraphNode->GetDocumentationLink(), DocumentationExcerpt);
+	}
+
+	TSharedRef<SButton> OpenGameplayEffectButton = SNew(SButton)
+		.ContentPadding(0.0f)
+		.ButtonStyle(FEditorStyle::Get(), "NoBorder")
+		.OnClicked(this, &SGraphNodeK2GameplayEffectVar::OnOpenGameplayEffect)
+		.ToolTipText(ButtonTooltipText)
+		.ToolTip(Tooltip)
+		.Visibility(this, &SGraphNodeK2GameplayEffectVar::IsOpenGameplayEffectButtonVisible)
+		[
+			ButtonContent.ToSharedRef()
+		];
+
+	OpenGameplayEffectButton->SetCursor(EMouseCursor::Hand);
+
+	return OpenGameplayEffectButton;
+}
+
+FReply SGraphNodeK2GameplayEffectVar::OnOpenGameplayEffect()
+{
+
+	TArray<UObject*> ObjectsForPropertiesMenu;
+
+	if (ObjectsForPropertiesMenu.Num() > 0)
+	{
+		FAssetEditorManager::Get().OpenEditorForAssets(ObjectsForPropertiesMenu);
+	}
+
+	return FReply::Handled(); 
+}
+
+EVisibility SGraphNodeK2GameplayEffectVar::IsOpenGameplayEffectButtonVisible() const
+{
+	UK2Node_GameplayEffectVariable* GameplayEffectNode = Cast<UK2Node_GameplayEffectVariable>(GraphNode);
+	if (GameplayEffectNode)
+	{
+		// set this to be hidden if there is no gameplay effect associated with this node
+	}
+	return EVisibility::Visible;
+}
+
+const FSlateBrush* SGraphNodeK2GameplayEffectVar::GetShadowBrush(bool bSelected) const
+{
+	UK2Node* K2Node = CastChecked<UK2Node>(GraphNode);
+	const bool bCompactMode = K2Node->ShouldDrawCompact();
+
+	if (bCompactMode)
+	{
+		return bSelected ? FEditorStyle::GetBrush(TEXT("Graph.VarNode.ShadowSelected")) : FEditorStyle::GetBrush(TEXT("Graph.VarNode.Shadow"));
+	}
+	return SGraphNodeK2Base::GetShadowBrush(bSelected);
+}
+
+#undef LOCTEXT_NAMESPACE
diff --git a/Engine/Source/Editor/GameplayAbilitiesEditor/Private/SGraphNodeK2GameplayEffectVar.h b/Engine/Source/Editor/GameplayAbilitiesEditor/Private/SGraphNodeK2GameplayEffectVar.h
new file mode 100644
--- /dev/null
+++ b/Engine/Source/Editor/GameplayAbilitiesEditor/Private/SGraphNodeK2GameplayEffectVar.h
@@ -0,0 +1,33 @@
+// Copyright 1998-2014 Epic Games, Inc. All Rights Reserved.
+
+#pragma once
+
+#include "SNodePanel.h"
+#include "SGraphNode.h"
+#include "KismetNodes/SGraphNodeK2Base.h"
+#include "KismetNodes/SGraphNodeK2Var.h"
+
+class SGraphNodeK2GameplayEffectVar : public SGraphNodeK2Var
+{
+public:
+	SLATE_BEGIN_ARGS(SGraphNodeK2Var){}
+	SLATE_END_ARGS()
+
+	void Construct( const FArguments& InArgs, UK2Node* InNode );
+
+	virtual const FSlateBrush* GetShadowBrush(bool bSelected) const;
+
+	// SGraphNode interface
+	virtual void UpdateGraphNode() override;
+	// End of SGraphNode interface
+
+protected:
+	TSharedRef<SWidget> OpenGameplayEffectButtonContent(FText ButtonText, FText ButtonTooltipText, FString DocumentationExcerpt = FString());
+
+//	void OnOpenGameplayEffectButtonClicked();
+	
+	FReply OnOpenGameplayEffect();
+	EVisibility IsOpenGameplayEffectButtonVisible() const;
+
+	FSlateColor GetVariableColor() const;
+};
