Make the GameplayAbilityBlueprint an Ubergraph

diff --git a/Engine/Source/Editor/GameplayAbilitiesEditor/Private/GameplayAbilitiesBlueprintFactory.cpp b/Engine/Source/Editor/GameplayAbilitiesEditor/Private/GameplayAbilitiesBlueprintFactory.cpp
--- a/Engine/Source/Editor/GameplayAbilitiesEditor/Private/GameplayAbilitiesBlueprintFactory.cpp
+++ b/Engine/Source/Editor/GameplayAbilitiesEditor/Private/GameplayAbilitiesBlueprintFactory.cpp
@@ -280,7 +280,7 @@ UObject* UGameplayAbilitiesBlueprintFactory::FactoryCreateNew(UClass* Class, UOb
 
 				// Only allow a gameplay ability graph if there isn't one in a parent blueprint
 				UEdGraph* NewGraph = FBlueprintEditorUtils::CreateNewGraph(NewBP, TEXT("Gameplay Ability Graph"), UGameplayAbilityGraph::StaticClass(), UGameplayAbilityGraphSchema::StaticClass());
-				FBlueprintEditorUtils::AddDomainSpecificGraph(NewBP, NewGraph);
+				FBlueprintEditorUtils::AddUbergraphPage(NewBP, NewGraph);
 				NewBP->LastEditedDocuments.Add(NewGraph);
 				NewGraph->bAllowDeletion = false;
 			}
