Improvements to GameplayEffect Nodes in GameplayAbilityEditor
 - Enabled button to open GameplayEffect for editing
 - Display ChanceToApply and ChanceToExecute if they are < 1

diff --git a/Engine/Source/Editor/GameplayAbilitiesEditor/Classes/K2Node_GameplayEffectVariable.h b/Engine/Source/Editor/GameplayAbilitiesEditor/Classes/K2Node_GameplayEffectVariable.h
--- a/Engine/Source/Editor/GameplayAbilitiesEditor/Classes/K2Node_GameplayEffectVariable.h
+++ b/Engine/Source/Editor/GameplayAbilitiesEditor/Classes/K2Node_GameplayEffectVariable.h
@@ -24,14 +24,16 @@ public:
 			EGameplayModOp::Type ModOp;
 		};
 
-		FGameplayEffectDisplayInfo() {}
+		FGameplayEffectDisplayInfo() : GameplayEffect(NULL)
+		{}
 
-		void Init(const UGameplayEffect& GameplayEffect);
+		void Init(UGameplayEffect& GameplayEffect);
 
 		float Duration;
 		float Period;
 		float ChanceToApplyToTarget;
 		float ChanceToExecuteOnGameplayEffect;
+		UGameplayEffect* GameplayEffect;
 
 		TArray<ModifierDisplayInfo> AttributeModifiers;
 	};
diff --git a/Engine/Source/Editor/GameplayAbilitiesEditor/Private/GameplayAbilityGraphSchema.cpp b/Engine/Source/Editor/GameplayAbilitiesEditor/Private/GameplayAbilityGraphSchema.cpp
--- a/Engine/Source/Editor/GameplayAbilitiesEditor/Private/GameplayAbilityGraphSchema.cpp
+++ b/Engine/Source/Editor/GameplayAbilitiesEditor/Private/GameplayAbilityGraphSchema.cpp
@@ -24,14 +24,20 @@ UK2Node_VariableGet* UGameplayAbilityGraphSchema::SpawnVariableGetNode(const FVe
 		UObject* ActiveObject = GenClass->GetDefaultObject();
 
 		// If the variable is a GameplayEffect create a custom node to show it
-		UGameplayEffect* GameplayEffect = Cast<UGameplayEffect>(ObjProp->GetObjectPropertyValue_InContainer(ActiveObject));
-		if (GameplayEffect)
+		FString PropType;
+		ObjProp->GetCPPMacroType(PropType);
+		if (PropType == "UGameplayEffect")
 		{
 			UK2Node_GameplayEffectVariable* NodeTemplate = NewObject<UK2Node_GameplayEffectVariable>();
 			UEdGraphSchema_K2::ConfigureVarNode(NodeTemplate, VariableName, Source, Blueprint);
 
 			UK2Node_GameplayEffectVariable* VariableNode = FEdGraphSchemaAction_K2NewNode::SpawnNodeFromTemplate<UK2Node_GameplayEffectVariable>(ParentGraph, NodeTemplate, GraphPosition);
-			VariableNode->GameplayEffectInfo.Init(*GameplayEffect);
+
+			UGameplayEffect* GameplayEffect = Cast<UGameplayEffect>(ObjProp->GetObjectPropertyValue_InContainer(ActiveObject));
+			if (GameplayEffect)
+			{
+				VariableNode->GameplayEffectInfo.Init(*GameplayEffect);
+			}
 
 			return VariableNode;
 		}
diff --git a/Engine/Source/Editor/GameplayAbilitiesEditor/Private/K2Node_GameplayEffectVariable.cpp b/Engine/Source/Editor/GameplayAbilitiesEditor/Private/K2Node_GameplayEffectVariable.cpp
--- a/Engine/Source/Editor/GameplayAbilitiesEditor/Private/K2Node_GameplayEffectVariable.cpp
+++ b/Engine/Source/Editor/GameplayAbilitiesEditor/Private/K2Node_GameplayEffectVariable.cpp
@@ -11,8 +11,9 @@ UK2Node_GameplayEffectVariable::UK2Node_GameplayEffectVariable(const class FPost
 {
 }
 
-void UK2Node_GameplayEffectVariable::FGameplayEffectDisplayInfo::Init(const UGameplayEffect& GameplayEffect)
+void UK2Node_GameplayEffectVariable::FGameplayEffectDisplayInfo::Init(UGameplayEffect& GameplayEffect)
 {
+	this->GameplayEffect = &GameplayEffect;
 	Duration = GameplayEffect.Duration.Value;
 	Period = GameplayEffect.Period.Value;
 	ChanceToApplyToTarget = GameplayEffect.ChanceToApplyToTarget.Value;
diff --git a/Engine/Source/Editor/GameplayAbilitiesEditor/Private/SGraphNodeK2GameplayEffectVar.cpp b/Engine/Source/Editor/GameplayAbilitiesEditor/Private/SGraphNodeK2GameplayEffectVar.cpp
--- a/Engine/Source/Editor/GameplayAbilitiesEditor/Private/SGraphNodeK2GameplayEffectVar.cpp
+++ b/Engine/Source/Editor/GameplayAbilitiesEditor/Private/SGraphNodeK2GameplayEffectVar.cpp
@@ -77,6 +77,52 @@ void SGraphNodeK2GameplayEffectVar::UpdateGraphNode()
 			SNew(SSpacer)
 		];
 
+	if (GameplayEffectNode->GameplayEffectInfo.ChanceToApplyToTarget < 1.f || GameplayEffectNode->GameplayEffectInfo.ChanceToExecuteOnGameplayEffect < 1.f)
+	{
+		FString ApplyToTargetString = FString::Printf(TEXT("%d"), GameplayEffectNode->GameplayEffectInfo.ChanceToApplyToTarget);
+		FString ExecuteOnGEString = FString::Printf(TEXT("%d"), GameplayEffectNode->GameplayEffectInfo.ChanceToExecuteOnGameplayEffect);
+
+		RightNodeBox->AddSlot()
+			[
+				SNew(SHorizontalBox)
+				+ SHorizontalBox::Slot()
+				.HAlign(HAlign_Left)
+				[
+					SNew(STextBlock)
+					.Text(FString("Chance to Apply to Target:"))
+				]
+				+ SHorizontalBox::Slot()
+					.HAlign(HAlign_Right)
+					[
+						SNew(STextBlock)
+						.Text(ApplyToTargetString)
+					]
+			];
+
+		RightNodeBox->AddSlot()
+			[
+				SNew(SHorizontalBox)
+				+ SHorizontalBox::Slot()
+				.HAlign(HAlign_Left)
+				[
+					SNew(STextBlock)
+					.Text(FString("Chance to Execute on GE:"))
+				]
+				+ SHorizontalBox::Slot()
+					.HAlign(HAlign_Right)
+					[
+						SNew(STextBlock)
+						.Text(ExecuteOnGEString)
+					]
+			];
+
+		RightNodeBox->AddSlot()
+			.FillHeight(0.3f)
+			[
+				SNew(SSpacer)
+			];
+	}
+
 	for (int ModIdx = 0; ModIdx < GameplayEffectNode->GameplayEffectInfo.AttributeModifiers.Num(); ++ModIdx)
 	{
 		RightNodeBox->AddSlot()
@@ -120,17 +166,17 @@ void SGraphNodeK2GameplayEffectVar::UpdateGraphNode()
 			];
 	}
 
-// 	TSharedRef<SWidget> OpenGameplayEffectButton = OpenGameplayEffectButtonContent(
-// 		LOCTEXT("FormatTextNodeOpenGameplayEffectButton", "Open Gameplay Effect"),
-// 		LOCTEXT("FormatTextNodeOpenGameplayEffectButton_Tooltip", "Opens this gameplay effect so that it can be edited."));
-// 
-// 	LeftNodeBox->AddSlot()
-// 		.AutoHeight()
-// 		.VAlign(VAlign_Center)
-// 		.Padding(10, 10, 10, 4)
-// 		[
-// 			OpenGameplayEffectButton
-// 		];
+	TSharedRef<SWidget> OpenGameplayEffectButton = OpenGameplayEffectButtonContent(
+		LOCTEXT("FormatTextNodeOpenGameplayEffectButton", "Open Gameplay Effect"),
+		LOCTEXT("FormatTextNodeOpenGameplayEffectButton_Tooltip", "Opens this gameplay effect so that it can be edited."));
+
+	LeftNodeBox->AddSlot()
+		.AutoHeight()
+		.VAlign(VAlign_Center)
+		.Padding(10, 10, 10, 4)
+		[
+			OpenGameplayEffectButton
+		];
 }
 
 TSharedRef<SWidget> SGraphNodeK2GameplayEffectVar::OpenGameplayEffectButtonContent(FText ButtonText, FText ButtonTooltipText, FString DocumentationExcerpt)
@@ -181,12 +227,11 @@ TSharedRef<SWidget> SGraphNodeK2GameplayEffectVar::OpenGameplayEffectButtonConte
 
 FReply SGraphNodeK2GameplayEffectVar::OnOpenGameplayEffect()
 {
+	UK2Node_GameplayEffectVariable* GameplayEffectNode = Cast<UK2Node_GameplayEffectVariable>(GraphNode);
 
-	TArray<UObject*> ObjectsForPropertiesMenu;
-
-	if (ObjectsForPropertiesMenu.Num() > 0)
+	if (GameplayEffectNode && GameplayEffectNode->GameplayEffectInfo.GameplayEffect)
 	{
-		FAssetEditorManager::Get().OpenEditorForAssets(ObjectsForPropertiesMenu);
+		FAssetEditorManager::Get().OpenEditorForAsset(GameplayEffectNode->GameplayEffectInfo.GameplayEffect);
 	}
 
 	return FReply::Handled(); 
diff --git a/Engine/Source/Editor/GameplayAbilitiesEditor/Private/SGraphNodeK2GameplayEffectVar.h b/Engine/Source/Editor/GameplayAbilitiesEditor/Private/SGraphNodeK2GameplayEffectVar.h
--- a/Engine/Source/Editor/GameplayAbilitiesEditor/Private/SGraphNodeK2GameplayEffectVar.h
+++ b/Engine/Source/Editor/GameplayAbilitiesEditor/Private/SGraphNodeK2GameplayEffectVar.h
@@ -23,8 +23,6 @@ public:
 
 protected:
 	TSharedRef<SWidget> OpenGameplayEffectButtonContent(FText ButtonText, FText ButtonTooltipText, FString DocumentationExcerpt = FString());
-
-//	void OnOpenGameplayEffectButtonClicked();
 	
 	FReply OnOpenGameplayEffect();
 	EVisibility IsOpenGameplayEffectButtonVisible() const;
