fix nonunity

diff --git a/Engine/Source/Editor/GameplayAbilitiesEditor/Private/GameplayAbilityGraphSchema.cpp b/Engine/Source/Editor/GameplayAbilitiesEditor/Private/GameplayAbilityGraphSchema.cpp
--- a/Engine/Source/Editor/GameplayAbilitiesEditor/Private/GameplayAbilityGraphSchema.cpp
+++ b/Engine/Source/Editor/GameplayAbilitiesEditor/Private/GameplayAbilityGraphSchema.cpp
@@ -5,6 +5,8 @@
 #include "GameplayEffect.h"
 #include "GameplayAbilityGraphSchema.h"
 #include "Kismet/KismetSystemLibrary.h"
+#include "BlueprintEditorUtils.h"
+
 
 UGameplayAbilityGraphSchema::UGameplayAbilityGraphSchema(const FPostConstructInitializeProperties& PCIP)
 : Super(PCIP)
diff --git a/Engine/Source/Editor/GameplayAbilitiesEditor/Private/SGraphNodeK2GameplayEffectVar.cpp b/Engine/Source/Editor/GameplayAbilitiesEditor/Private/SGraphNodeK2GameplayEffectVar.cpp
--- a/Engine/Source/Editor/GameplayAbilitiesEditor/Private/SGraphNodeK2GameplayEffectVar.cpp
+++ b/Engine/Source/Editor/GameplayAbilitiesEditor/Private/SGraphNodeK2GameplayEffectVar.cpp
@@ -11,6 +11,7 @@
 #include "GameplayEffectTypes.h"
 #include "ClassIconFinder.h"
 #include "IDocumentation.h"
+#include "AssetEditorManager.h"
 
 #define LOCTEXT_NAMESPACE "GraphNodeGameplayEffectVar"
 
