GameplayAbilities - updates to allow literals to be specified in EXposeOnSpawn properties in AbilityTasks which spawn actors

diff --git a/Engine/Source/Editor/GameplayAbilitiesEditor/Private/K2Node_LatentAbilityCall.cpp b/Engine/Source/Editor/GameplayAbilitiesEditor/Private/K2Node_LatentAbilityCall.cpp
--- a/Engine/Source/Editor/GameplayAbilitiesEditor/Private/K2Node_LatentAbilityCall.cpp
+++ b/Engine/Source/Editor/GameplayAbilitiesEditor/Private/K2Node_LatentAbilityCall.cpp
@@ -7,6 +7,7 @@
 #include "BlueprintEditorUtils.h"
 #include "K2Node_LatentAbilityCall.h"
 #include "GameplayAbilityGraphSchema.h"
+#include "K2Node_EnumLiteral.h"
 
 #define LOCTEXT_NAMESPACE "K2Node"
 
@@ -132,6 +133,21 @@ void UK2Node_LatentAbilityCall::CreatePinsForClass(UClass* InClass)
 
 	SpawnParmPins.Empty();
 
+	// Tasks can hide spawn parameters by doing meta = (HideSpawnParms="PropertyA,PropertyB")
+	// (For example, hide Instigator in situations where instigator is not relevant to your task)
+	
+	TArray<FString> IgnorePropertyList;
+	{
+		UFunction* ProxyFunction = ProxyFactoryClass->FindFunctionByName(ProxyFactoryFunctionName);
+
+		FString IgnorePropertyListStr = ProxyFunction->GetMetaData(FName(TEXT("HideSpawnParms")));
+	
+		if (!IgnorePropertyListStr.IsEmpty())
+		{
+			IgnorePropertyListStr.ParseIntoArray(&IgnorePropertyList, TEXT(","), true);
+		}
+	}
+
 	for (TFieldIterator<UProperty> PropertyIt(InClass, EFieldIteratorFlags::IncludeSuper); PropertyIt; ++PropertyIt)
 	{
 		UProperty* Property = *PropertyIt;
@@ -144,11 +160,13 @@ void UK2Node_LatentAbilityCall::CreatePinsForClass(UClass* InClass)
 			!Property->HasAnyPropertyFlags(CPF_Parm) &&
 			bIsSettableExternally &&
 			Property->HasAllPropertyFlags(CPF_BlueprintVisible) &&
-			!bIsDelegate)
+			!bIsDelegate && 
+			!IgnorePropertyList.Contains(Property->GetName()))
 		{
+
+
 			UEdGraphPin* Pin = CreatePin(EGPD_Input, TEXT(""), TEXT(""), NULL, false, false, Property->GetName());
 			const bool bPinGood = (Pin != NULL) && K2Schema->ConvertPropertyToPinType(Property, /*out*/ Pin->PinType);
-			Pin->bDefaultValueIsIgnored = true;
 			SpawnParmPins.Add(Pin);
 		}
 	}
@@ -417,10 +435,10 @@ void UK2Node_LatentAbilityCall::ExpandNode(class FKismetCompilerContext& Compile
 	// -------------------------------------------
 	// Set spawn variables
 	// -------------------------------------------
-
+	
 	for (auto SpawnVarPin : SpawnParmPins)
 	{
-		if (SpawnVarPin->LinkedTo.Num() > 0)
+		if (SpawnVarPin->LinkedTo.Num() > 0 || SpawnVarPin->DefaultValue != FString())
 		{
 			UFunction* SetByNameFunction = Schema->FindSetVariableByNameFunction(SpawnVarPin->PinType);
 			if (SetByNameFunction)
@@ -452,14 +470,32 @@ void UK2Node_LatentAbilityCall::ExpandNode(class FKismetCompilerContext& Compile
 				// Fill in literal for 'property name' pin - name of pin is property name
 				UEdGraphPin* PropertyNamePin = SetVarNode->FindPinChecked(PropertyNameParamName);
 				PropertyNamePin->DefaultValue = SpawnVarPin->PinName;
-
-				// Move connection from the variable pin on the spawn node to the 'value' pin
+				
 				UEdGraphPin* ValuePin = SetVarNode->FindPinChecked(ValueParamName);
-				CompilerContext.MovePinLinksToIntermediate(*SpawnVarPin, *ValuePin);
-				SetVarNode->PinConnectionListChanged(ValuePin);
+				if (SpawnVarPin->LinkedTo.Num() == 0 &&
+					SpawnVarPin->DefaultValue != FString() &&
+					SpawnVarPin->PinType.PinCategory == Schema->PC_Byte &&
+					SpawnVarPin->PinType.PinSubCategoryObject.IsValid() &&
+					SpawnVarPin->PinType.PinSubCategoryObject->IsA<UEnum>())
+				{
+					// Pin is an enum, we need to alias the enum value to an int:
+					UK2Node_EnumLiteral* EnumLiteralNode = CompilerContext.SpawnIntermediateNode<UK2Node_EnumLiteral>(this, SourceGraph);
+					EnumLiteralNode->Enum = CastChecked<UEnum>(SpawnVarPin->PinType.PinSubCategoryObject.Get());
+					EnumLiteralNode->AllocateDefaultPins();
+					EnumLiteralNode->FindPinChecked(Schema->PN_ReturnValue)->MakeLinkTo(ValuePin);
+
+					CompilerContext.MovePinLinksToIntermediate(*SpawnVarPin, *EnumLiteralNode->FindPin(TEXT("Enum")));
+					SetVarNode->PinConnectionListChanged(ValuePin);					
+				}
+				else
+				{
+					// Move connection from the variable pin on the spawn node to the 'value' pin
+					CompilerContext.MovePinLinksToIntermediate(*SpawnVarPin, *ValuePin);
+					SetVarNode->PinConnectionListChanged(ValuePin);
+				}
 			}
 		}
-	}
+	}	
 		
 
 	// -------------------------------------------
diff --git a/Engine/Source/Runtime/GameplayAbilities/Public/Abilities/Tasks/AbilityTask_PlayMontageAndWait.h b/Engine/Source/Runtime/GameplayAbilities/Public/Abilities/Tasks/AbilityTask_PlayMontageAndWait.h
--- a/Engine/Source/Runtime/GameplayAbilities/Public/Abilities/Tasks/AbilityTask_PlayMontageAndWait.h
+++ b/Engine/Source/Runtime/GameplayAbilities/Public/Abilities/Tasks/AbilityTask_PlayMontageAndWait.h
@@ -20,7 +20,7 @@ class UAbilityTask_PlayMontageAndWait : public UAbilityTask
 	void OnMontageEnded(UAnimMontage* Montage, bool bInterrupted);
 
 
-	UFUNCTION(BlueprintCallable, meta = (HidePin = "WorldContextObject", DefaultToSelf = "WorldContextObject", BlueprintInternalUseOnly = "TRUE"))
+	UFUNCTION(BlueprintCallable, meta = (FriendlyName="PlayMontageAndWait", HidePin = "WorldContextObject", DefaultToSelf = "WorldContextObject", BlueprintInternalUseOnly = "TRUE"))
 	static UAbilityTask_PlayMontageAndWait* CreatePlayMontageAndWaitProxy(UObject* WorldContextObject, UAnimMontage *MontageToPlay);
 
 	virtual void Activate() override;
diff --git a/Engine/Source/Runtime/GameplayAbilities/Public/Abilities/Tasks/AbilityTask_WaitTargetData.h b/Engine/Source/Runtime/GameplayAbilities/Public/Abilities/Tasks/AbilityTask_WaitTargetData.h
--- a/Engine/Source/Runtime/GameplayAbilities/Public/Abilities/Tasks/AbilityTask_WaitTargetData.h
+++ b/Engine/Source/Runtime/GameplayAbilities/Public/Abilities/Tasks/AbilityTask_WaitTargetData.h
@@ -32,7 +32,7 @@ class UAbilityTask_WaitTargetData: public UAbilityTask
 	void OnTargetDataCancelledCallback(FGameplayAbilityTargetDataHandle Data);
 
 	/** Spawns Targeting actor and waits for it to return valid data or to be cancelled. */
-	UFUNCTION(BlueprintCallable, meta=(HidePin = "WorldContextObject", DefaultToSelf = "WorldContextObject", BlueprintInternalUseOnly = "true"), Category = "Abilities")
+	UFUNCTION(BlueprintCallable, meta=(HidePin = "WorldContextObject", DefaultToSelf = "WorldContextObject", BlueprintInternalUseOnly = "true", HideSpawnParms="Instigator"), Category = "Abilities")
 	static UAbilityTask_WaitTargetData* WaitTargetData(UObject* WorldContextObject, TSubclassOf<AGameplayAbilityTargetActor> Class);
 
 	UFUNCTION(BlueprintCallable, meta = (HidePin = "WorldContextObject", DefaultToSelf = "WorldContextObject", BlueprintInternalUseOnly = "true"), Category = "Abilities")
diff --git a/Engine/Source/Runtime/GameplayAbilities/Public/AbilitySystemBlueprintLibrary.h b/Engine/Source/Runtime/GameplayAbilities/Public/AbilitySystemBlueprintLibrary.h
--- a/Engine/Source/Runtime/GameplayAbilities/Public/AbilitySystemBlueprintLibrary.h
+++ b/Engine/Source/Runtime/GameplayAbilities/Public/AbilitySystemBlueprintLibrary.h
@@ -12,6 +12,7 @@ class UAbilityTask_WaitMovementModeChange;
 class UAbilityTask_WaitOverlap;
 class UAbilityTask_WaitConfirmCancel;
 
+// meta =(RestrictedToClasses="GameplayAbility")
 UCLASS(MinimalAPI)
 class UAbilitySystemBlueprintLibrary : public UBlueprintFunctionLibrary
 {
