Update LatentAbilityCall enum literal value passing

diff --git a/Engine/Source/Editor/GameplayAbilitiesEditor/Private/K2Node_LatentAbilityCall.cpp b/Engine/Source/Editor/GameplayAbilitiesEditor/Private/K2Node_LatentAbilityCall.cpp
--- a/Engine/Source/Editor/GameplayAbilitiesEditor/Private/K2Node_LatentAbilityCall.cpp
+++ b/Engine/Source/Editor/GameplayAbilitiesEditor/Private/K2Node_LatentAbilityCall.cpp
@@ -484,8 +484,8 @@ void UK2Node_LatentAbilityCall::ExpandNode(class FKismetCompilerContext& Compile
 					EnumLiteralNode->AllocateDefaultPins();
 					EnumLiteralNode->FindPinChecked(Schema->PN_ReturnValue)->MakeLinkTo(ValuePin);
 
-					CompilerContext.MovePinLinksToIntermediate(*SpawnVarPin, *EnumLiteralNode->FindPin(TEXT("Enum")));
-					SetVarNode->PinConnectionListChanged(ValuePin);					
+					UEdGraphPin* InPin = EnumLiteralNode->FindPinChecked(UK2Node_EnumLiteral::GetEnumInputPinName());
+					InPin->DefaultValue = SpawnVarPin->DefaultValue;
 				}
 				else
 				{
