remove default event graph from gameplay ability blueprints

diff --git a/Engine/Source/Editor/GameplayAbilitiesEditor/Private/GameplayAbilitiesBlueprintFactory.cpp b/Engine/Source/Editor/GameplayAbilitiesEditor/Private/GameplayAbilitiesBlueprintFactory.cpp
--- a/Engine/Source/Editor/GameplayAbilitiesEditor/Private/GameplayAbilitiesBlueprintFactory.cpp
+++ b/Engine/Source/Editor/GameplayAbilitiesEditor/Private/GameplayAbilitiesBlueprintFactory.cpp
@@ -280,6 +280,12 @@ UObject* UGameplayAbilitiesBlueprintFactory::FactoryCreateNew(UClass* Class, UOb
 
 				// Only allow a gameplay ability graph if there isn't one in a parent blueprint
 				UEdGraph* NewGraph = FBlueprintEditorUtils::CreateNewGraph(NewBP, TEXT("Gameplay Ability Graph"), UGameplayAbilityGraph::StaticClass(), UGameplayAbilityGraphSchema::StaticClass());
+#if WITH_EDITORONLY_DATA
+				if (NewBP->UbergraphPages.Num())
+				{
+					FBlueprintEditorUtils::RemoveGraphs(NewBP, NewBP->UbergraphPages);
+				}
+#endif
 				FBlueprintEditorUtils::AddUbergraphPage(NewBP, NewGraph);
 				NewBP->LastEditedDocuments.Add(NewGraph);
 				NewGraph->bAllowDeletion = false;
diff --git a/Engine/Source/Editor/GameplayAbilitiesEditor/Private/GameplayAbilitiesEditor.cpp b/Engine/Source/Editor/GameplayAbilitiesEditor/Private/GameplayAbilitiesEditor.cpp
--- a/Engine/Source/Editor/GameplayAbilitiesEditor/Private/GameplayAbilitiesEditor.cpp
+++ b/Engine/Source/Editor/GameplayAbilitiesEditor/Private/GameplayAbilitiesEditor.cpp
@@ -27,6 +27,28 @@ FGameplayAbilitiesEditor::~FGameplayAbilitiesEditor()
 void FGameplayAbilitiesEditor::InitGameplayAbilitiesEditor(const EToolkitMode::Type Mode, const TSharedPtr< IToolkitHost >& InitToolkitHost, const TArray<UBlueprint*>& InBlueprints, bool bShouldOpenInDefaultsMode)
 {
 	InitBlueprintEditor(Mode, InitToolkitHost, InBlueprints, bShouldOpenInDefaultsMode);
+
+	for (auto Blueprint : InBlueprints)
+	{
+		EnsureGameplayAbilityBlueprintIsUpToDate(Blueprint);
+	}
+}
+
+void FGameplayAbilitiesEditor::EnsureGameplayAbilityBlueprintIsUpToDate(UBlueprint* Blueprint)
+{
+#if WITH_EDITORONLY_DATA
+	int32 Count = Blueprint->UbergraphPages.Num();
+	for (auto Graph : Blueprint->UbergraphPages)
+	{
+		// remove the default event graph, if it exists, from existing Gameplay Ability Blueprints
+		if (Graph->GetName() == "EventGraph" && Graph->Nodes.Num() == 0)
+		{
+			check(!Graph->Schema->GetClass()->IsChildOf(UGameplayAbilityGraphSchema::StaticClass()));
+			FBlueprintEditorUtils::RemoveGraph(Blueprint, Graph);
+			break;
+		}
+	}
+#endif
 }
 
 FName FGameplayAbilitiesEditor::GetToolkitFName() const
diff --git a/Engine/Source/Editor/GameplayAbilitiesEditor/Public/GameplayAbilitiesEditor.h b/Engine/Source/Editor/GameplayAbilitiesEditor/Public/GameplayAbilitiesEditor.h
--- a/Engine/Source/Editor/GameplayAbilitiesEditor/Public/GameplayAbilitiesEditor.h
+++ b/Engine/Source/Editor/GameplayAbilitiesEditor/Public/GameplayAbilitiesEditor.h
@@ -25,6 +25,15 @@ public:
 	 */ 
 
 	void InitGameplayAbilitiesEditor(const EToolkitMode::Type Mode, const TSharedPtr<IToolkitHost>& InitToolkitHost, const TArray<UBlueprint*>& InBlueprints, bool bShouldOpenInDefaultsMode);
+
+private:
+	/**
+	 * Updates existing gameplay ability blueprints to make sure that they are up to date
+	 * 
+	 * @param	Blueprint	The blueprint to be updated
+	 */
+	void EnsureGameplayAbilityBlueprintIsUpToDate(UBlueprint* Blueprint);
+
 public:
 	FGameplayAbilitiesEditor();
 
diff --git a/Engine/Source/Editor/UnrealEd/Private/Kismet2/BlueprintEditorUtils.cpp b/Engine/Source/Editor/UnrealEd/Private/Kismet2/BlueprintEditorUtils.cpp
--- a/Engine/Source/Editor/UnrealEd/Private/Kismet2/BlueprintEditorUtils.cpp
+++ b/Engine/Source/Editor/UnrealEd/Private/Kismet2/BlueprintEditorUtils.cpp
@@ -1774,6 +1774,7 @@ void FBlueprintEditorUtils::RemoveGraph(UBlueprint* Blueprint, class UEdGraph* G
 			Blueprint->DelegateSignatureGraphs.Remove( GraphToRemove );
 			Blueprint->FunctionGraphs.Remove( GraphToRemove );
 			Blueprint->UbergraphPages.Remove( GraphToRemove );
+			Blueprint->LastEditedDocuments.Remove( GraphToRemove );
 			if(Blueprint->MacroGraphs.Remove( GraphToRemove ) > 0 ) 
 			{
 				//removes all macro nodes using this macro graph
