#ue4 non-unity build fix

diff --git a/Engine/Source/Editor/GameplayAbilitiesEditor/Private/GameplayAbilitiesEditor.cpp b/Engine/Source/Editor/GameplayAbilitiesEditor/Private/GameplayAbilitiesEditor.cpp
--- a/Engine/Source/Editor/GameplayAbilitiesEditor/Private/GameplayAbilitiesEditor.cpp
+++ b/Engine/Source/Editor/GameplayAbilitiesEditor/Private/GameplayAbilitiesEditor.cpp
@@ -4,6 +4,8 @@
 #include "GameplayAbilitiesEditor.h"
 
 #include "GameplayAbilityBlueprint.h"
+#include "GameplayAbilityGraphSchema.h"
+#include "Kismet2/BlueprintEditorUtils.h"
 
 #define LOCTEXT_NAMESPACE "FGameplayAbilitiesEditor"
 
