GameplayAbilities
-Fix so GameplayCue events don't appear on non IGameplayCue actor blueprint graphs.
-Fix so that AbilityTasks show up in GameplayAbility macros (but not non UGameplayAbility actors)

diff --git a/Engine/Source/Editor/GameplayAbilitiesEditor/Private/K2Node_GameplayCueEvent.cpp b/Engine/Source/Editor/GameplayAbilitiesEditor/Private/K2Node_GameplayCueEvent.cpp
--- a/Engine/Source/Editor/GameplayAbilitiesEditor/Private/K2Node_GameplayCueEvent.cpp
+++ b/Engine/Source/Editor/GameplayAbilitiesEditor/Private/K2Node_GameplayCueEvent.cpp
@@ -37,6 +37,21 @@ void UK2Node_GameplayCueEvent::GetMenuEntries(FGraphContextMenuBuilder& Context)
 {
 	Super::GetMenuEntries(Context);
 
+	bool bValidClass = false;
+	UBlueprint* MyBlueprint = Cast<UBlueprint>(Context.CurrentGraph->GetOuter());
+	if (MyBlueprint && MyBlueprint->GeneratedClass)
+	{
+		if (MyBlueprint->GeneratedClass->ImplementsInterface(UGameplayCueInterface::StaticClass()))
+		{
+			bValidClass = true;
+		}
+	}
+
+	if (!bValidClass)
+	{
+		return;
+	}
+
 	const FString FunctionCategory(TEXT("GameplayCue Event"));
 
 	IGameplayTagsModule& GameplayTagsModule = IGameplayTagsModule::Get();
@@ -45,11 +60,17 @@ void UK2Node_GameplayCueEvent::GetMenuEntries(FGraphContextMenuBuilder& Context)
 	FGameplayTagContainer CueTags = GameplayTagsModule.GetGameplayTagsManager().RequestGameplayTagChildren(RootTag);
 
 	// Fixme: need to check if this function is already defined so that it can be reimplemented
+	//	-Checking MyBlueprint->GeneratedClass isn't enough since they may have added an event and not recompiled
+	//	-FEdGraphSchemaAction_K2AddCustomEvent does not check names/always ensures a valid name
+	//	-FEdGraphSchemaAction_K2AddEvent does check and recenters - but it looks at EventSignatureName/EventSignatureClass for equality and that
+	//		won't work here.
+	//	
+	//	Probably need a new EdGraphSchemaAction to do this properly. For now this is ok since they will get a compile error if they do drop in
+	//	two of the same GameplayCue even Nodes and it should be pretty clear that they can't do that.
 
 	for (auto It = CueTags.CreateConstIterator(); It; ++It)
 	{
 		FGameplayTag Tag = *It;
-
 		UK2Node_GameplayCueEvent* NodeTemplate = Context.CreateTemplateNode<UK2Node_GameplayCueEvent>();
 
 		NodeTemplate->CustomFunctionName = Tag.GetTagName();
@@ -64,4 +85,4 @@ void UK2Node_GameplayCueEvent::GetMenuEntries(FGraphContextMenuBuilder& Context)
 	}	
 }
 
-#undef LOCTEXT_NAMESPACE
\ No newline at end of file
+#undef LOCTEXT_NAMESPACE
diff --git a/Engine/Source/Editor/GameplayAbilitiesEditor/Private/K2Node_LatentAbilityCall.cpp b/Engine/Source/Editor/GameplayAbilitiesEditor/Private/K2Node_LatentAbilityCall.cpp
--- a/Engine/Source/Editor/GameplayAbilitiesEditor/Private/K2Node_LatentAbilityCall.cpp
+++ b/Engine/Source/Editor/GameplayAbilitiesEditor/Private/K2Node_LatentAbilityCall.cpp
@@ -19,11 +19,26 @@ UK2Node_LatentAbilityCall::UK2Node_LatentAbilityCall(const class FPostConstructI
 
 bool UK2Node_LatentAbilityCall::CanCreateUnderSpecifiedSchema(const UEdGraphSchema* DesiredSchema) const
 {
-	return DesiredSchema->GetClass()->IsChildOf(UGameplayAbilityGraphSchema::StaticClass());
+	return Super::CanCreateUnderSpecifiedSchema(DesiredSchema);
 }
 
 void UK2Node_LatentAbilityCall::GetMenuEntries(FGraphContextMenuBuilder& ContextMenuBuilder) const
 {
+	bool bValidClass = false;
+	UBlueprint* MyBlueprint = Cast<UBlueprint>(ContextMenuBuilder.CurrentGraph->GetOuter());
+	if (MyBlueprint && MyBlueprint->GeneratedClass)
+	{
+		if (MyBlueprint->GeneratedClass->IsChildOf(UGameplayAbility::StaticClass()))
+		{
+			bValidClass = true;
+		}
+	}
+
+	if (!bValidClass)
+	{
+		return;
+	}
+
 	const UEdGraphSchema_K2* K2Schema = GetDefault<UEdGraphSchema_K2>();
 	EGraphType GraphType = K2Schema->GetGraphType(ContextMenuBuilder.CurrentGraph);
 	const bool bAllowLatentFuncs = (GraphType == GT_Ubergraph || GraphType == GT_Macro);
