(WIP) Setting up async task nodes so that they're listed for the new blueprint action menus.

diff --git a/Engine/Source/Editor/BlueprintGraph/Classes/K2Node_AIMoveTo.h b/Engine/Source/Editor/BlueprintGraph/Classes/K2Node_AIMoveTo.h
--- a/Engine/Source/Editor/BlueprintGraph/Classes/K2Node_AIMoveTo.h
+++ b/Engine/Source/Editor/BlueprintGraph/Classes/K2Node_AIMoveTo.h
@@ -16,5 +16,7 @@ class BLUEPRINTGRAPH_API UK2Node_AIMoveTo : public UK2Node_BaseAsyncTask
 	virtual FText GetNodeTitle(ENodeTitleType::Type TitleType) const override;
 	// End UEdGraphNode interface
 
-	virtual FString GetCategoryName();
+	// Begin UK2Node interface
+	virtual FText GetMenuCategory() const override;
+	// End UK2Node interface
 };
diff --git a/Engine/Source/Editor/BlueprintGraph/Classes/K2Node_BaseAsyncTask.h b/Engine/Source/Editor/BlueprintGraph/Classes/K2Node_BaseAsyncTask.h
--- a/Engine/Source/Editor/BlueprintGraph/Classes/K2Node_BaseAsyncTask.h
+++ b/Engine/Source/Editor/BlueprintGraph/Classes/K2Node_BaseAsyncTask.h
@@ -26,11 +26,11 @@ class BLUEPRINTGRAPH_API UK2Node_BaseAsyncTask : public UK2Node
 	virtual void ExpandNode(class FKismetCompilerContext& CompilerContext, UEdGraph* SourceGraph) override;
 	virtual bool HasExternalBlueprintDependencies(TArray<class UStruct*>* OptionalOutput) const override;
 	virtual FName GetCornerIcon() const override;
+	virtual void GetMenuActions(TArray<UBlueprintNodeSpawner*>& ActionListOut) const override;
+	virtual FText GetMenuCategory() const override;
 	// End of UK2Node interface
 
 protected:
-	virtual FString GetCategoryName();
-
 	// Creates a default menu entry
 	TSharedPtr<FEdGraphSchemaAction_K2NewNode> CreateDefaultMenuEntry(UK2Node_BaseAsyncTask* NodeTemplate, FGraphContextMenuBuilder& ContextMenuBuilder) const;
 
diff --git a/Engine/Source/Editor/BlueprintGraph/Private/BlueprintFunctionNodeSpawner.cpp b/Engine/Source/Editor/BlueprintGraph/Private/BlueprintFunctionNodeSpawner.cpp
--- a/Engine/Source/Editor/BlueprintGraph/Private/BlueprintFunctionNodeSpawner.cpp
+++ b/Engine/Source/Editor/BlueprintGraph/Private/BlueprintFunctionNodeSpawner.cpp
@@ -55,8 +55,13 @@ UBlueprintFunctionNodeSpawner::UBlueprintFunctionNodeSpawner(class FPostConstruc
 //------------------------------------------------------------------------------
 UEdGraphNode* UBlueprintFunctionNodeSpawner::Invoke(UEdGraph* ParentGraph) const
 {
-	UK2Node_CallFunction* NewNode = CastChecked<UK2Node_CallFunction>(Super::Invoke(ParentGraph));
-	NewNode->SetFromFunction(Function);
+	UEdGraphNode* NewNode = Super::Invoke(ParentGraph);
+	// user could have changed the node class (to something like
+	// UK2Node_BaseAsyncTask, which also wraps a function)
+	if (UK2Node_CallFunction* FuncNode = Cast<UK2Node_CallFunction>(NewNode))
+	{
+		FuncNode->SetFromFunction(Function);
+	}
 	
 	bool bIsTemplateNode = ParentGraph->HasAnyFlags(RF_Transient);
 	if (CustomizeNodeDelegate.IsBound())
diff --git a/Engine/Source/Editor/BlueprintGraph/Private/K2Node_AIMoveTo.cpp b/Engine/Source/Editor/BlueprintGraph/Private/K2Node_AIMoveTo.cpp
--- a/Engine/Source/Editor/BlueprintGraph/Private/K2Node_AIMoveTo.cpp
+++ b/Engine/Source/Editor/BlueprintGraph/Private/K2Node_AIMoveTo.cpp
@@ -6,6 +6,7 @@
 #include "Blueprint/AIBlueprintHelperLibrary.h"
 #include "Blueprint/AIAsyncTaskBlueprintProxy.h"
 #include "K2Node_AIMoveTo.h"
+#include "EditorCategoryUtils.h"
 
 #define LOCTEXT_NAMESPACE "K2Node_AIMoveTo"
 
@@ -17,9 +18,9 @@ UK2Node_AIMoveTo::UK2Node_AIMoveTo(const class FPostConstructInitializePropertie
 	ProxyClass = UAIAsyncTaskBlueprintProxy::StaticClass();
 }
 
-FString UK2Node_AIMoveTo::GetCategoryName()
+FText UK2Node_AIMoveTo::GetMenuCategory() const
 {
-	return TEXT("AI");
+	return FEditorCategoryUtils::GetCommonCategory(FCommonEditorCategory::Ai);
 }
 
 FString UK2Node_AIMoveTo::GetTooltip() const
diff --git a/Engine/Source/Editor/BlueprintGraph/Private/K2Node_BaseAsyncTask.cpp b/Engine/Source/Editor/BlueprintGraph/Private/K2Node_BaseAsyncTask.cpp
--- a/Engine/Source/Editor/BlueprintGraph/Private/K2Node_BaseAsyncTask.cpp
+++ b/Engine/Source/Editor/BlueprintGraph/Private/K2Node_BaseAsyncTask.cpp
@@ -6,6 +6,7 @@
 #include "CompilerResultsLog.h"
 #include "KismetCompiler.h"
 #include "K2ActionMenuBuilder.h"
+#include "BlueprintNodeSpawner.h"
 
 #define LOCTEXT_NAMESPACE "UK2Node_BaseAsyncTask"
 
@@ -30,12 +31,6 @@ FText UK2Node_BaseAsyncTask::GetNodeTitle(ENodeTitleType::Type TitleType) const
 	return FText::FromString(FunctionToolTipText);
 }
 
-FString UK2Node_BaseAsyncTask::GetCategoryName()
-{
-	const FString GenericFunctionCategory = LOCTEXT("FunctionCategory", "Call Function").ToString();
-	return UK2Node_CallFunction::GetDefaultCategoryForFunction(GetFactoryFunction(), GenericFunctionCategory);
-}
-
 void UK2Node_BaseAsyncTask::GetMenuEntries(FGraphContextMenuBuilder& ContextMenuBuilder) const
 {
 	const UEdGraphSchema_K2* K2Schema = GetDefault<UEdGraphSchema_K2>();
@@ -50,7 +45,7 @@ void UK2Node_BaseAsyncTask::GetMenuEntries(FGraphContextMenuBuilder& ContextMenu
 
 TSharedPtr<FEdGraphSchemaAction_K2NewNode> UK2Node_BaseAsyncTask::CreateDefaultMenuEntry(UK2Node_BaseAsyncTask* NodeTemplate, FGraphContextMenuBuilder& ContextMenuBuilder) const
 {
-	TSharedPtr<FEdGraphSchemaAction_K2NewNode> NodeAction = FK2ActionMenuBuilder::AddNewNodeAction(ContextMenuBuilder, NodeTemplate->GetCategoryName(), NodeTemplate->GetNodeTitle(ENodeTitleType::ListView), NodeTemplate->GetTooltip(), 0, NodeTemplate->GetKeywords());
+	TSharedPtr<FEdGraphSchemaAction_K2NewNode> NodeAction = FK2ActionMenuBuilder::AddNewNodeAction(ContextMenuBuilder, NodeTemplate->GetMenuCategory().ToString(), NodeTemplate->GetNodeTitle(ENodeTitleType::ListView), NodeTemplate->GetTooltip(), 0, NodeTemplate->GetKeywords());
 	
 	NodeAction->NodeTemplate = NodeTemplate;
 
@@ -358,9 +353,23 @@ bool UK2Node_BaseAsyncTask::HasExternalBlueprintDependencies(TArray<class UStruc
 }
 
 FName UK2Node_BaseAsyncTask::GetCornerIcon() const
-	{
+{
 	return TEXT("Graph.Latent.LatentIcon");
-	}
+}
+
+FText UK2Node_BaseAsyncTask::GetMenuCategory() const
+{	
+	UFunction* TargetFunction = GetFactoryFunction();
+	return FText::FromString(UK2Node_CallFunction::GetDefaultCategoryForFunction(TargetFunction, TEXT("Call Function")));
+}
+
+void UK2Node_BaseAsyncTask::GetMenuActions(TArray<UBlueprintNodeSpawner*>& ActionListOut) const
+{
+	UBlueprintNodeSpawner* NodeSpawner = UBlueprintNodeSpawner::Create(GetClass());
+	check(NodeSpawner != nullptr);
+	
+	ActionListOut.Add(NodeSpawner);
+}
 
 UFunction* UK2Node_BaseAsyncTask::GetFactoryFunction() const
 {
diff --git a/Engine/Source/Editor/GameplayAbilitiesEditor/Classes/K2Node_LatentAbilityCall.h b/Engine/Source/Editor/GameplayAbilitiesEditor/Classes/K2Node_LatentAbilityCall.h
--- a/Engine/Source/Editor/GameplayAbilitiesEditor/Classes/K2Node_LatentAbilityCall.h
+++ b/Engine/Source/Editor/GameplayAbilitiesEditor/Classes/K2Node_LatentAbilityCall.h
@@ -16,6 +16,7 @@ class UK2Node_LatentAbilityCall : public UK2Node_BaseAsyncTask
 	// UEdGraphNode interface
 	virtual bool CanCreateUnderSpecifiedSchema(const UEdGraphSchema* DesiredSchema) const;
 	virtual void GetMenuEntries(FGraphContextMenuBuilder& ContextMenuBuilder) const override;
+	virtual void GetMenuActions(TArray<UBlueprintNodeSpawner*>& ActionListOut) const override;
 	// End of UEdGraphNode interface
 
 	virtual void AllocateDefaultPins() override;
diff --git a/Engine/Source/Editor/GameplayAbilitiesEditor/Private/K2Node_LatentAbilityCall.cpp b/Engine/Source/Editor/GameplayAbilitiesEditor/Private/K2Node_LatentAbilityCall.cpp
--- a/Engine/Source/Editor/GameplayAbilitiesEditor/Private/K2Node_LatentAbilityCall.cpp
+++ b/Engine/Source/Editor/GameplayAbilitiesEditor/Private/K2Node_LatentAbilityCall.cpp
@@ -8,6 +8,7 @@
 #include "K2Node_LatentAbilityCall.h"
 #include "GameplayAbilityGraphSchema.h"
 #include "K2Node_EnumLiteral.h"
+#include "BlueprintFunctionNodeSpawner.h"
 
 #define LOCTEXT_NAMESPACE "K2Node"
 
@@ -74,6 +75,62 @@ void UK2Node_LatentAbilityCall::GetMenuEntries(FGraphContextMenuBuilder& Context
 	}
 }
 
+void UK2Node_LatentAbilityCall::GetMenuActions(TArray<UBlueprintNodeSpawner*>& ActionListOut) const
+{
+	// these nested loops are combing over the same classes/functions the
+	// FBlueprintActionDatabase does; ideally we save on perf and fold this in
+	// with FBlueprintActionDatabase, but we want to keep the modules separate
+	for (TObjectIterator<UClass> ClassIt; ClassIt; ++ClassIt)
+	{
+		UClass* Class = *ClassIt;
+		if (!Class->IsChildOf<UAbilityTask>() || Class->HasAnyClassFlags(CLASS_Abstract))
+		{
+			continue;
+		}
+		
+		for (TFieldIterator<UFunction> FuncIt(Class, EFieldIteratorFlags::ExcludeSuper); FuncIt; ++FuncIt)
+		{
+			UFunction* Function = *FuncIt;
+			if (!Function->HasAnyFunctionFlags(FUNC_Static))
+			{
+				continue;
+			}
+			
+			UObjectProperty* ReturnProperty = Cast<UObjectProperty>(Function->GetReturnProperty());
+			// see if the function is a static factory method for online proxies
+			bool const bIsProxyFactoryMethod = (ReturnProperty != nullptr) && ReturnProperty->PropertyClass->IsChildOf<UAbilityTask>();
+			
+			if (bIsProxyFactoryMethod)
+			{
+				UBlueprintNodeSpawner* NodeSpawner = UBlueprintFunctionNodeSpawner::Create(Function);
+				check(NodeSpawner != nullptr);
+				NodeSpawner->NodeClass = GetClass();
+				
+				auto CustomizeAcyncNodeLambda = [](UEdGraphNode* NewNode, bool bIsTemplateNode, TWeakObjectPtr<UFunction> FunctionPtr)
+				{
+					UK2Node_LatentAbilityCall* AsyncTaskNode = CastChecked<UK2Node_LatentAbilityCall>(NewNode);
+					if (FunctionPtr.IsValid())
+					{
+						UFunction* Function = FunctionPtr.Get();
+						UObjectProperty* ReturnProperty = CastChecked<UObjectProperty>(Function->GetReturnProperty());
+						
+						AsyncTaskNode->ProxyFactoryFunctionName = Function->GetFName();
+						AsyncTaskNode->ProxyFactoryClass        = Function->GetOuterUClass();
+						AsyncTaskNode->ProxyClass               = ReturnProperty->PropertyClass;
+					}
+				};
+				
+				TWeakObjectPtr<UFunction> FunctionPtr = Function;
+				NodeSpawner->CustomizeNodeDelegate = UBlueprintNodeSpawner::FCustomizeNodeDelegate::CreateStatic(CustomizeAcyncNodeLambda, FunctionPtr);
+				
+				// @TODO: since this can't be folded into FBlueprintActionDatabase, we
+				//        need a way to associate these spawners with a certain class
+				ActionListOut.Add(NodeSpawner);
+			}
+		}
+	}
+}
+
 // -------------------------------------------------
 
 struct FK2Node_LatentAbilityCallHelper
diff --git a/Engine/Source/Editor/OnlineBlueprintSupport/Classes/K2Node_LatentOnlineCall.h b/Engine/Source/Editor/OnlineBlueprintSupport/Classes/K2Node_LatentOnlineCall.h
--- a/Engine/Source/Editor/OnlineBlueprintSupport/Classes/K2Node_LatentOnlineCall.h
+++ b/Engine/Source/Editor/OnlineBlueprintSupport/Classes/K2Node_LatentOnlineCall.h
@@ -9,8 +9,9 @@ UCLASS()
 class ONLINEBLUEPRINTSUPPORT_API UK2Node_LatentOnlineCall : public UK2Node_BaseAsyncTask
 {
 	GENERATED_UCLASS_BODY()
-
+	
 	// UEdGraphNode interface
 	virtual void GetMenuEntries(FGraphContextMenuBuilder& ContextMenuBuilder) const override;
+	virtual void GetMenuActions(TArray<UBlueprintNodeSpawner*>& ActionListOut) const override;
 	// End of UEdGraphNode interface
 };
diff --git a/Engine/Source/Editor/OnlineBlueprintSupport/Classes/K2Node_LeaderboardFlush.h b/Engine/Source/Editor/OnlineBlueprintSupport/Classes/K2Node_LeaderboardFlush.h
--- a/Engine/Source/Editor/OnlineBlueprintSupport/Classes/K2Node_LeaderboardFlush.h
+++ b/Engine/Source/Editor/OnlineBlueprintSupport/Classes/K2Node_LeaderboardFlush.h
@@ -14,7 +14,7 @@ class ONLINEBLUEPRINTSUPPORT_API UK2Node_LeaderboardFlush : public UK2Node_BaseA
 	virtual FText GetNodeTitle(ENodeTitleType::Type TitleType) const override;
 	// End of UEdGraphNode interface
 
-	// UK2Node_BaseAsyncTask interface
-	virtual FString GetCategoryName() override;
-	// End of UK2Node_BaseAsyncTask interface
+	// Begin UK2Node interface
+	virtual FText GetMenuCategory() const override;
+	// End UK2Node interface
 };
diff --git a/Engine/Source/Editor/OnlineBlueprintSupport/Private/K2Node_LatentOnlineCall.cpp b/Engine/Source/Editor/OnlineBlueprintSupport/Private/K2Node_LatentOnlineCall.cpp
--- a/Engine/Source/Editor/OnlineBlueprintSupport/Private/K2Node_LatentOnlineCall.cpp
+++ b/Engine/Source/Editor/OnlineBlueprintSupport/Private/K2Node_LatentOnlineCall.cpp
@@ -3,6 +3,7 @@
 #include "OnlineBlueprintSupportPrivatePCH.h"
 #include "OnlineBlueprintCallProxyBase.h"
 #include "K2Node_LatentOnlineCall.h"
+#include "BlueprintFunctionNodeSpawner.h"
 
 #define LOCTEXT_NAMESPACE "K2Node"
 
@@ -49,4 +50,60 @@ void UK2Node_LatentOnlineCall::GetMenuEntries(FGraphContextMenuBuilder& ContextM
 	}
 }
 
+void UK2Node_LatentOnlineCall::GetMenuActions(TArray<UBlueprintNodeSpawner*>& ActionListOut) const
+{
+	// these nested loops are combing over the same classes/functions the
+	// FBlueprintActionDatabase does; ideally we save on perf and fold this in
+	// with FBlueprintActionDatabase, but we want to keep the modules separate
+	for (TObjectIterator<UClass> ClassIt; ClassIt; ++ClassIt)
+	{
+		UClass* Class = *ClassIt;
+		if (!Class->IsChildOf<UOnlineBlueprintCallProxyBase>() || Class->HasAnyClassFlags(CLASS_Abstract))
+		{
+			continue;
+		}
+		
+		for (TFieldIterator<UFunction> FuncIt(Class, EFieldIteratorFlags::ExcludeSuper); FuncIt; ++FuncIt)
+		{
+			UFunction* Function = *FuncIt;
+			if (!Function->HasAnyFunctionFlags(FUNC_Static))
+			{
+				continue;
+			}
+			
+			UObjectProperty* ReturnProperty = Cast<UObjectProperty>(Function->GetReturnProperty());
+			// see if the function is a static factory method for online proxies
+			bool const bIsProxyFactoryMethod = (ReturnProperty != nullptr) && ReturnProperty->PropertyClass->IsChildOf<UOnlineBlueprintCallProxyBase>();
+			
+			if (bIsProxyFactoryMethod)
+			{
+				UBlueprintNodeSpawner* NodeSpawner = UBlueprintFunctionNodeSpawner::Create(Function);
+				check(NodeSpawner != nullptr);
+				NodeSpawner->NodeClass = GetClass();
+				
+				auto CustomizeAcyncNodeLambda = [](UEdGraphNode* NewNode, bool bIsTemplateNode, TWeakObjectPtr<UFunction> FunctionPtr)
+				{
+					UK2Node_LatentOnlineCall* AsyncTaskNode = CastChecked<UK2Node_LatentOnlineCall>(NewNode);
+					if (FunctionPtr.IsValid())
+					{
+						UFunction* Function = FunctionPtr.Get();
+						UObjectProperty* ReturnProperty = CastChecked<UObjectProperty>(Function->GetReturnProperty());
+						
+						AsyncTaskNode->ProxyFactoryFunctionName = Function->GetFName();
+						AsyncTaskNode->ProxyFactoryClass        = Function->GetOuterUClass();
+						AsyncTaskNode->ProxyClass               = ReturnProperty->PropertyClass;
+					}
+				};
+				
+				TWeakObjectPtr<UFunction> FunctionPtr = Function;
+				NodeSpawner->CustomizeNodeDelegate = UBlueprintNodeSpawner::FCustomizeNodeDelegate::CreateStatic(CustomizeAcyncNodeLambda, FunctionPtr);
+				
+				// @TODO: since this can't be folded into FBlueprintActionDatabase, we
+				//        need a way to associate these spawners with a certain class
+				ActionListOut.Add(NodeSpawner);
+			}
+		}
+	}
+}
+
 #undef LOCTEXT_NAMESPACE
diff --git a/Engine/Source/Editor/OnlineBlueprintSupport/Private/K2Node_LeaderboardFlush.cpp b/Engine/Source/Editor/OnlineBlueprintSupport/Private/K2Node_LeaderboardFlush.cpp
--- a/Engine/Source/Editor/OnlineBlueprintSupport/Private/K2Node_LeaderboardFlush.cpp
+++ b/Engine/Source/Editor/OnlineBlueprintSupport/Private/K2Node_LeaderboardFlush.cpp
@@ -25,9 +25,9 @@ FText UK2Node_LeaderboardFlush::GetNodeTitle(ENodeTitleType::Type TitleType) con
 	return LOCTEXT("FlushLeaderboards", "Flush Leaderboards");
 }
 
-FString UK2Node_LeaderboardFlush::GetCategoryName()
+FText UK2Node_LeaderboardFlush::GetMenuCategory() const
 {
-	return TEXT("Online|Leaderboard");
+	return LOCTEXT("LeaderboardFlushCategory", "Online|Leaderboard");
 }
 
 #undef LOCTEXT_NAMESPACE
diff --git a/Engine/Source/Runtime/GameplayAbilities/Public/Abilities/Tasks/AbilityTask_PlayMontageAndWait.h b/Engine/Source/Runtime/GameplayAbilities/Public/Abilities/Tasks/AbilityTask_PlayMontageAndWait.h
--- a/Engine/Source/Runtime/GameplayAbilities/Public/Abilities/Tasks/AbilityTask_PlayMontageAndWait.h
+++ b/Engine/Source/Runtime/GameplayAbilities/Public/Abilities/Tasks/AbilityTask_PlayMontageAndWait.h
@@ -20,7 +20,7 @@ class UAbilityTask_PlayMontageAndWait : public UAbilityTask
 	void OnMontageEnded(UAnimMontage* Montage, bool bInterrupted);
 
 
-	UFUNCTION(BlueprintCallable, meta = (FriendlyName="PlayMontageAndWait", HidePin = "WorldContextObject", DefaultToSelf = "WorldContextObject", BlueprintInternalUseOnly = "TRUE"))
+	UFUNCTION(BlueprintCallable, Category=Abilities, meta = (FriendlyName="PlayMontageAndWait", HidePin = "WorldContextObject", DefaultToSelf = "WorldContextObject", BlueprintInternalUseOnly = "TRUE"))
 	static UAbilityTask_PlayMontageAndWait* CreatePlayMontageAndWaitProxy(UObject* WorldContextObject, UAnimMontage *MontageToPlay);
 
 	virtual void Activate() override;
diff --git a/Engine/Source/Runtime/GameplayAbilities/Public/Abilities/Tasks/AbilityTask_WaitAbilityActivate.h b/Engine/Source/Runtime/GameplayAbilities/Public/Abilities/Tasks/AbilityTask_WaitAbilityActivate.h
--- a/Engine/Source/Runtime/GameplayAbilities/Public/Abilities/Tasks/AbilityTask_WaitAbilityActivate.h
+++ b/Engine/Source/Runtime/GameplayAbilities/Public/Abilities/Tasks/AbilityTask_WaitAbilityActivate.h
@@ -26,7 +26,7 @@ class UAbilityTask_WaitAbilityActivate : public UAbilityTask
 	void OnAbilityActivate(UGameplayAbility *ActivatedAbility);
 
 	/** Wait until a new ability (of the same or different type) is activated. */
-	UFUNCTION(BlueprintCallable, meta = (HidePin = "WorldContextObject", DefaultToSelf = "WorldContextObject", BlueprintInternalUseOnly = "TRUE"))
+	UFUNCTION(BlueprintCallable, Category=Abilities, meta = (HidePin = "WorldContextObject", DefaultToSelf = "WorldContextObject", BlueprintInternalUseOnly = "TRUE"))
 	static UAbilityTask_WaitAbilityActivate* WaitForAbilityActivate(UObject* WorldContextObject, FGameplayTag WithTag, FGameplayTag WithoutTag);
 
 	FGameplayTag WithTag;
diff --git a/Engine/Source/Runtime/GameplayAbilities/Public/Abilities/Tasks/AbilityTask_WaitAbilityCommit.h b/Engine/Source/Runtime/GameplayAbilities/Public/Abilities/Tasks/AbilityTask_WaitAbilityCommit.h
--- a/Engine/Source/Runtime/GameplayAbilities/Public/Abilities/Tasks/AbilityTask_WaitAbilityCommit.h
+++ b/Engine/Source/Runtime/GameplayAbilities/Public/Abilities/Tasks/AbilityTask_WaitAbilityCommit.h
@@ -28,7 +28,7 @@ class UAbilityTask_WaitAbilityCommit : public UAbilityTask
 	void OnAbilityCommit(UGameplayAbility *ActivatedAbility);
 
 	/** Wait until a new ability (of the same or different type) is commited. Used to gracefully interrupt abilities in specific ways. */
-	UFUNCTION(BlueprintCallable, meta = (HidePin = "WorldContextObject", DefaultToSelf = "WorldContextObject", BlueprintInternalUseOnly = "TRUE", FriendlyName = "Wait For New Ability Commit"))
+	UFUNCTION(BlueprintCallable, Category=Abilities, meta = (HidePin = "WorldContextObject", DefaultToSelf = "WorldContextObject", BlueprintInternalUseOnly = "TRUE", FriendlyName = "Wait For New Ability Commit"))
 	static UAbilityTask_WaitAbilityCommit* WaitForAbilityCommit(UObject* WorldContextObject, FGameplayTag WithTag, FGameplayTag WithoutTage);	
 
 	FGameplayTag WithTag;
diff --git a/Engine/Source/Runtime/GameplayAbilities/Public/Abilities/Tasks/AbilityTask_WaitAttributeChange.h b/Engine/Source/Runtime/GameplayAbilities/Public/Abilities/Tasks/AbilityTask_WaitAttributeChange.h
--- a/Engine/Source/Runtime/GameplayAbilities/Public/Abilities/Tasks/AbilityTask_WaitAttributeChange.h
+++ b/Engine/Source/Runtime/GameplayAbilities/Public/Abilities/Tasks/AbilityTask_WaitAttributeChange.h
@@ -26,7 +26,7 @@ class UAbilityTask_WaitAttributeChange : public UAbilityTask
 	void OnAttributeChange(float NewValue, const FGameplayEffectModCallbackData*);
 
 	/** Wait until an attribute changes. */
-	UFUNCTION(BlueprintCallable, meta = (HidePin = "WorldContextObject", DefaultToSelf = "WorldContextObject", BlueprintInternalUseOnly = "TRUE"))
+	UFUNCTION(BlueprintCallable, Category=Abilities, meta = (HidePin = "WorldContextObject", DefaultToSelf = "WorldContextObject", BlueprintInternalUseOnly = "TRUE"))
 	static UAbilityTask_WaitAttributeChange* WaitForAttributeChange(UObject* WorldContextObject, FGameplayAttribute Attribute, FGameplayTag WithTag, FGameplayTag WithoutTag);
 
 	FGameplayTag WithTag;
diff --git a/Engine/Source/Runtime/GameplayAbilities/Public/Abilities/Tasks/AbilityTask_WaitGameplayEffectRemoved.h b/Engine/Source/Runtime/GameplayAbilities/Public/Abilities/Tasks/AbilityTask_WaitGameplayEffectRemoved.h
--- a/Engine/Source/Runtime/GameplayAbilities/Public/Abilities/Tasks/AbilityTask_WaitGameplayEffectRemoved.h
+++ b/Engine/Source/Runtime/GameplayAbilities/Public/Abilities/Tasks/AbilityTask_WaitGameplayEffectRemoved.h
@@ -28,7 +28,7 @@ class UAbilityTask_WaitGameplayEffectRemoved : public UAbilityTask
 	void OnGameplayEffectRemoved();
 
 	/** Wait until an overlap occurs. This will need to be better fleshed out so we can specify game specific collision requirements */
-	UFUNCTION(BlueprintCallable, meta = (HidePin = "WorldContextObject", DefaultToSelf = "WorldContextObject", BlueprintInternalUseOnly = "TRUE"))
+	UFUNCTION(BlueprintCallable, Category=Abilities, meta = (HidePin = "WorldContextObject", DefaultToSelf = "WorldContextObject", BlueprintInternalUseOnly = "TRUE"))
 	static UAbilityTask_WaitGameplayEffectRemoved* WaitForGameplayEffectRemoved(UObject* WorldContextObject, FActiveGameplayEffectHandle Handle);
 
 	FActiveGameplayEffectHandle Handle;
diff --git a/Engine/Source/Runtime/GameplayAbilities/Public/Abilities/Tasks/AbilityTask_WaitMovementModeChange.h b/Engine/Source/Runtime/GameplayAbilities/Public/Abilities/Tasks/AbilityTask_WaitMovementModeChange.h
--- a/Engine/Source/Runtime/GameplayAbilities/Public/Abilities/Tasks/AbilityTask_WaitMovementModeChange.h
+++ b/Engine/Source/Runtime/GameplayAbilities/Public/Abilities/Tasks/AbilityTask_WaitMovementModeChange.h
@@ -20,7 +20,7 @@ public:
 
 	EMovementMode	RequiredMode;
 
-	UFUNCTION(BlueprintCallable, meta = (HidePin = "WorldContextObject", DefaultToSelf = "WorldContextObject", BlueprintInternalUseOnly = "TRUE"))
+	UFUNCTION(BlueprintCallable, Category=Abilities, meta = (HidePin = "WorldContextObject", DefaultToSelf = "WorldContextObject", BlueprintInternalUseOnly = "TRUE"))
 	static UAbilityTask_WaitMovementModeChange* CreateWaitMovementModeChange(UObject* WorldContextObject, EMovementMode NewMode);
 
 	virtual void Activate() override;
diff --git a/Engine/Source/Runtime/GameplayAbilities/Public/Abilities/Tasks/AbilityTask_WaitOverlap.h b/Engine/Source/Runtime/GameplayAbilities/Public/Abilities/Tasks/AbilityTask_WaitOverlap.h
--- a/Engine/Source/Runtime/GameplayAbilities/Public/Abilities/Tasks/AbilityTask_WaitOverlap.h
+++ b/Engine/Source/Runtime/GameplayAbilities/Public/Abilities/Tasks/AbilityTask_WaitOverlap.h
@@ -31,7 +31,7 @@ class UAbilityTask_WaitOverlap : public UAbilityTask
 	virtual void Activate() override;
 
 	/** Wait until an overlap occurs. This will need to be better fleshed out so we can specify game specific collision requirements */
-	UFUNCTION(BlueprintCallable, meta = (HidePin = "WorldContextObject", DefaultToSelf = "WorldContextObject", BlueprintInternalUseOnly = "TRUE"))
+	UFUNCTION(BlueprintCallable, Category=Abilities, meta = (HidePin = "WorldContextObject", DefaultToSelf = "WorldContextObject", BlueprintInternalUseOnly = "TRUE"))
 	static UAbilityTask_WaitOverlap* WaitForOverlap(UObject* WorldContextObject);
 
 public:
