Fix issues with GameplayEffectVariable nodes not storing and loading information correctly.

diff --git a/Engine/Source/Editor/GameplayAbilitiesEditor/Classes/K2Node_GameplayEffectVariable.h b/Engine/Source/Editor/GameplayAbilitiesEditor/Classes/K2Node_GameplayEffectVariable.h
--- a/Engine/Source/Editor/GameplayAbilitiesEditor/Classes/K2Node_GameplayEffectVariable.h
+++ b/Engine/Source/Editor/GameplayAbilitiesEditor/Classes/K2Node_GameplayEffectVariable.h
@@ -8,35 +8,12 @@
 
 #include "K2Node_GameplayEffectVariable.generated.h"
 
-UCLASS(MinimalAPI)
-class UK2Node_GameplayEffectVariable : public UK2Node_VariableGet // only need a custom node for the getter
+UCLASS()
+class GAMEPLAYABILITIESEDITOR_API UK2Node_GameplayEffectVariable : public UK2Node_VariableGet // only need a custom node for the getter
 {
 	GENERATED_UCLASS_BODY()
 
 public:
-	class FGameplayEffectDisplayInfo
-	{
-	public:
-		struct ModifierDisplayInfo
-		{
-			FString DisplayName;
-			float Magnitude;
-			EGameplayModOp::Type ModOp;
-		};
-
-		FGameplayEffectDisplayInfo() : GameplayEffect(NULL)
-		{}
-
-		void Init(UGameplayEffect& GameplayEffect);
-
-		float Duration;
-		float Period;
-		float ChanceToApplyToTarget;
-		float ChanceToExecuteOnGameplayEffect;
-		UGameplayEffect* GameplayEffect;
-
-		TArray<ModifierDisplayInfo> AttributeModifiers;
-	};
-
-	FGameplayEffectDisplayInfo GameplayEffectInfo;
+	UPROPERTY()
+	UGameplayEffect* GameplayEffect;
 };
diff --git a/Engine/Source/Editor/GameplayAbilitiesEditor/Private/GameplayAbilityGraphSchema.cpp b/Engine/Source/Editor/GameplayAbilitiesEditor/Private/GameplayAbilityGraphSchema.cpp
--- a/Engine/Source/Editor/GameplayAbilitiesEditor/Private/GameplayAbilityGraphSchema.cpp
+++ b/Engine/Source/Editor/GameplayAbilitiesEditor/Private/GameplayAbilityGraphSchema.cpp
@@ -38,7 +38,7 @@ UK2Node_VariableGet* UGameplayAbilityGraphSchema::SpawnVariableGetNode(const FVe
 			UGameplayEffect* GameplayEffect = Cast<UGameplayEffect>(ObjProp->GetObjectPropertyValue_InContainer(ActiveObject));
 			if (GameplayEffect)
 			{
-				VariableNode->GameplayEffectInfo.Init(*GameplayEffect);
+				VariableNode->GameplayEffect = GameplayEffect;
 			}
 
 			return VariableNode;
diff --git a/Engine/Source/Editor/GameplayAbilitiesEditor/Private/K2Node_GameplayEffectVariable.cpp b/Engine/Source/Editor/GameplayAbilitiesEditor/Private/K2Node_GameplayEffectVariable.cpp
--- a/Engine/Source/Editor/GameplayAbilitiesEditor/Private/K2Node_GameplayEffectVariable.cpp
+++ b/Engine/Source/Editor/GameplayAbilitiesEditor/Private/K2Node_GameplayEffectVariable.cpp
@@ -11,26 +11,4 @@ UK2Node_GameplayEffectVariable::UK2Node_GameplayEffectVariable(const class FPost
 {
 }
 
-void UK2Node_GameplayEffectVariable::FGameplayEffectDisplayInfo::Init(UGameplayEffect& GameplayEffect)
-{
-	this->GameplayEffect = &GameplayEffect;
-	Duration = GameplayEffect.Duration.Value;
-	Period = GameplayEffect.Period.Value;
-	ChanceToApplyToTarget = GameplayEffect.ChanceToApplyToTarget.Value;
-	ChanceToExecuteOnGameplayEffect = GameplayEffect.ChanceToExecuteOnGameplayEffect.Value;
-
-	for (int Idx = 0; Idx < GameplayEffect.Modifiers.Num(); ++Idx)
-	{
-		if (GameplayEffect.Modifiers[Idx].ModifierType == EGameplayMod::Attribute)
-		{
-			ModifierDisplayInfo ModInfo;
-			ModInfo.DisplayName = GameplayEffect.Modifiers[Idx].Attribute.GetName();
-			ModInfo.Magnitude = GameplayEffect.Modifiers[Idx].Magnitude.Value;
-			ModInfo.ModOp = GameplayEffect.Modifiers[Idx].ModifierOp;
-
-			AttributeModifiers.Add(ModInfo);
-		}
-	}
-}
-
 #undef LOCTEXT_NAMESPACE
diff --git a/Engine/Source/Editor/GameplayAbilitiesEditor/Private/SGraphNodeK2GameplayEffectVar.cpp b/Engine/Source/Editor/GameplayAbilitiesEditor/Private/SGraphNodeK2GameplayEffectVar.cpp
--- a/Engine/Source/Editor/GameplayAbilitiesEditor/Private/SGraphNodeK2GameplayEffectVar.cpp
+++ b/Engine/Source/Editor/GameplayAbilitiesEditor/Private/SGraphNodeK2GameplayEffectVar.cpp
@@ -12,6 +12,7 @@
 #include "ClassIconFinder.h"
 #include "IDocumentation.h"
 #include "AssetEditorManager.h"
+#include "GameplayEffect.h"
 
 #define LOCTEXT_NAMESPACE "GraphNodeGameplayEffectVar"
 
@@ -34,9 +35,10 @@ void SGraphNodeK2GameplayEffectVar::UpdateGraphNode()
 	SGraphNodeK2Var::UpdateGraphNode();
 
 	UK2Node_GameplayEffectVariable* GameplayEffectNode = Cast<UK2Node_GameplayEffectVariable>(GraphNode);
+	UGameplayEffect* GameplayEffect = GameplayEffectNode ? GameplayEffectNode->GameplayEffect : NULL;
 
-	FString DurationString = FString::Printf(TEXT("%d s"), GameplayEffectNode->GameplayEffectInfo.Duration);
-	FString PeriodString = FString::Printf(TEXT("%d s"), GameplayEffectNode->GameplayEffectInfo.Period);
+	FString DurationString = FString::Printf(TEXT("%d s"), (GameplayEffect) ? GameplayEffect->Duration.Value : 0.f);
+	FString PeriodString = FString::Printf(TEXT("%d s"), (GameplayEffect) ? GameplayEffect->Period.Value : 0.f);
 
 	RightNodeBox->AddSlot()
 		[
@@ -78,10 +80,10 @@ void SGraphNodeK2GameplayEffectVar::UpdateGraphNode()
 			SNew(SSpacer)
 		];
 
-	if (GameplayEffectNode->GameplayEffectInfo.ChanceToApplyToTarget < 1.f || GameplayEffectNode->GameplayEffectInfo.ChanceToExecuteOnGameplayEffect < 1.f)
+	if (GameplayEffect && (GameplayEffect->ChanceToApplyToTarget.Value < 1.f || GameplayEffect->ChanceToExecuteOnGameplayEffect.Value < 1.f))
 	{
-		FString ApplyToTargetString = FString::Printf(TEXT("%d"), GameplayEffectNode->GameplayEffectInfo.ChanceToApplyToTarget);
-		FString ExecuteOnGEString = FString::Printf(TEXT("%d"), GameplayEffectNode->GameplayEffectInfo.ChanceToExecuteOnGameplayEffect);
+		FString ApplyToTargetString = FString::Printf(TEXT("%d"), GameplayEffect->ChanceToApplyToTarget.Value);
+		FString ExecuteOnGEString = FString::Printf(TEXT("%d"), GameplayEffect->ChanceToExecuteOnGameplayEffect.Value);
 
 		RightNodeBox->AddSlot()
 			[
@@ -124,8 +126,14 @@ void SGraphNodeK2GameplayEffectVar::UpdateGraphNode()
 			];
 	}
 
-	for (int ModIdx = 0; ModIdx < GameplayEffectNode->GameplayEffectInfo.AttributeModifiers.Num(); ++ModIdx)
+	const int32 ModifierCount = GameplayEffect ? GameplayEffect->Modifiers.Num() : 0;
+	for (int32 Idx = 0; Idx < ModifierCount; ++Idx)
 	{
+		if (GameplayEffect->Modifiers[Idx].ModifierType != EGameplayMod::Attribute)
+		{
+			continue;
+		}
+
 		RightNodeBox->AddSlot()
 			.FillHeight(0.3f)
 			[
@@ -137,11 +145,11 @@ void SGraphNodeK2GameplayEffectVar::UpdateGraphNode()
 			.HAlign(HAlign_Left)
 			[
 				SNew(STextBlock)
-				.Text(GameplayEffectNode->GameplayEffectInfo.AttributeModifiers[ModIdx].DisplayName)
+				.Text(GameplayEffect->Modifiers[Idx].Attribute.GetName())
 			];
 
 		FString ModOpString;
-		switch (GameplayEffectNode->GameplayEffectInfo.AttributeModifiers[ModIdx].ModOp)
+		switch (GameplayEffect->Modifiers[Idx].ModifierOp)
 		{
 		case EGameplayModOp::Additive:
 			ModOpString = "Add";
@@ -156,7 +164,7 @@ void SGraphNodeK2GameplayEffectVar::UpdateGraphNode()
 			ModOpString = "Other";
 			break;
 		}
-		FString ModString = FString::Printf(TEXT("%s: %f"), *ModOpString, GameplayEffectNode->GameplayEffectInfo.AttributeModifiers[ModIdx].Magnitude);
+		FString ModString = FString::Printf(TEXT("%s: %f"), *ModOpString, GameplayEffect->Modifiers[Idx].Magnitude.Value);
 
 		RightNodeBox->AddSlot()
 			.VAlign(VAlign_Center)
@@ -230,9 +238,9 @@ FReply SGraphNodeK2GameplayEffectVar::OnOpenGameplayEffect()
 {
 	UK2Node_GameplayEffectVariable* GameplayEffectNode = Cast<UK2Node_GameplayEffectVariable>(GraphNode);
 
-	if (GameplayEffectNode && GameplayEffectNode->GameplayEffectInfo.GameplayEffect)
+	if (GameplayEffectNode && GameplayEffectNode->GameplayEffect)
 	{
-		FAssetEditorManager::Get().OpenEditorForAsset(GameplayEffectNode->GameplayEffectInfo.GameplayEffect);
+		FAssetEditorManager::Get().OpenEditorForAsset(GameplayEffectNode->GameplayEffect);
 	}
 
 	return FReply::Handled(); 
