GameplayEffectTemplates

diff --git a/Engine/Source/Editor/GameplayAbilitiesEditor/Private/GameplayEffectDetails.cpp b/Engine/Source/Editor/GameplayAbilitiesEditor/Private/GameplayEffectDetails.cpp
new file mode 100644
--- /dev/null
+++ b/Engine/Source/Editor/GameplayAbilitiesEditor/Private/GameplayEffectDetails.cpp
@@ -0,0 +1,147 @@
+// Copyright 1998-2013 Epic Games, Inc. All Rights Reserved.
+
+#include "AbilitySystemEditorPrivatePCH.h"
+#include "GameplayEffectDetails.h"
+#include "Editor/PropertyEditor/Public/PropertyEditing.h"
+#include "BlueprintUtilities.h"
+#include "KismetEditorUtilities.h"
+#include "AbilitySystemGlobals.h"
+#include "GameplayAbilitiesModule.h"
+#include "SGameplayAttributeGraphPin.h"
+#include "GameplayEffectTemplate.h"
+
+#define LOCTEXT_NAMESPACE "GameplayEffectDetailsCustomization"
+
+DEFINE_LOG_CATEGORY(LogGameplayEffectDetails);
+
+// --------------------------------------------------------SGameplayAttributeGraphPin----------------------------
+
+TSharedRef<IDetailCustomization> FGameplayEffectDetails::MakeInstance()
+{
+	return MakeShareable(new FGameplayEffectDetails);
+}
+
+BEGIN_SLATE_FUNCTION_BUILD_OPTIMIZATION
+void FGameplayEffectDetails::CustomizeDetails(IDetailLayoutBuilder& DetailLayout)
+{
+	MyDetailLayout = &DetailLayout;
+
+	TArray< TWeakObjectPtr<UObject> > Objects;
+	DetailLayout.GetObjectsBeingCustomized(Objects);
+	if (Objects.Num() != 1)
+	{
+		// I don't know what to do here or what could be expected. Just return to disable all templating functionality
+		return;
+	}
+
+	TemplateProperty = DetailLayout.GetProperty(FName("Template"), UGameplayEffect::StaticClass());
+	ShowAllProperty = DetailLayout.GetProperty(FName("ShowAllProperties"), UGameplayEffect::StaticClass());
+
+	FSimpleDelegate UpdateShowAllDelegate = FSimpleDelegate::CreateSP(this, &FGameplayEffectDetails::OnShowAllChange);
+	ShowAllProperty->SetOnPropertyValueChanged(UpdateShowAllDelegate);
+
+	FSimpleDelegate UpdatTemplateDelegate = FSimpleDelegate::CreateSP(this, &FGameplayEffectDetails::OnTemplateChange);
+	TemplateProperty->SetOnPropertyValueChanged(UpdatTemplateDelegate);
+	
+	
+	// Hide properties where necessary
+	UGameplayEffect* Obj = Cast<UGameplayEffect>(Objects[0].Get());
+	if (Obj && !Obj->ShowAllProperties)
+	{
+		UGameplayEffectTemplate* Template = Obj->Template;
+
+		UE_LOG(LogGameplayEffectDetails, Warning, TEXT("ObjectBeingEditor: %s"), *Obj->GetName());
+		if (Template)
+		{
+			UGameplayEffect* DefObj = Obj->Template->GetClass()->GetDefaultObject<UGameplayEffect>();
+
+			for (TFieldIterator<UProperty> PropIt(UGameplayEffect::StaticClass(), EFieldIteratorFlags::ExcludeSuper); PropIt; ++PropIt)
+			{
+				UProperty* Property = *PropIt;
+				TSharedPtr<IPropertyHandle> PropHandle = DetailLayout.GetProperty(Property->GetFName());
+				HideProperties(DetailLayout, PropHandle, Template);
+			}
+		}
+	}
+}
+END_SLATE_FUNCTION_BUILD_OPTIMIZATION
+
+/** Recursively hide properties that are not default editable according to the Template */
+bool FGameplayEffectDetails::HideProperties(IDetailLayoutBuilder& DetailLayout, TSharedPtr<IPropertyHandle> PropHandle, UGameplayEffectTemplate* Template)
+{
+	// Never hide the Template or ShowAllProperties properties
+	if (TemplateProperty == PropHandle || ShowAllProperty == PropHandle)
+	{
+		return false;
+	}
+	
+	UProperty* UProp = PropHandle->GetProperty();
+
+	// Don't hide the EditableProperties
+	for (FString MatchStr : Template->EditableProperties)
+	{
+		if (MatchStr.Equals(UProp->GetName(), ESearchCase::IgnoreCase))
+		{			
+			return false;
+		}
+	}
+	
+	// Recurse into children - if they are all hidden then we are hidden.
+	uint32 NumChildren = 0;
+	PropHandle->GetNumChildren(NumChildren);
+
+	bool AllChildrenHidden = true;
+	for (uint32 ChildIdx = 0; ChildIdx < NumChildren; ++ChildIdx)
+	{
+		TSharedPtr<IPropertyHandle> ChildHandle = PropHandle->GetChildHandle(ChildIdx);
+		AllChildrenHidden &= HideProperties(DetailLayout, ChildHandle, Template);
+	}
+
+	if (AllChildrenHidden)
+	{
+		DetailLayout.HideProperty(PropHandle);
+		return true;
+	}
+
+	return false;
+}
+
+void FGameplayEffectDetails::OnShowAllChange()
+{
+	MyDetailLayout->ForceRefreshDetails();
+}
+
+void FGameplayEffectDetails::OnTemplateChange()
+{
+	TArray< TWeakObjectPtr<UObject> > Objects;
+	MyDetailLayout->GetObjectsBeingCustomized(Objects);
+	if (Objects.Num() != 1)
+	{
+		return;
+	}
+
+	UGameplayEffect* Obj = Cast<UGameplayEffect>(Objects[0].Get());
+	UGameplayEffect* Template = Obj->Template;
+
+	if (Template)
+	{
+		// Copy any non-default properties from the template into the current editable object
+		UGameplayEffect* DefObj = Template->GetClass()->GetDefaultObject<UGameplayEffect>();
+		for (TFieldIterator<UProperty> PropIt(UGameplayEffect::StaticClass(), EFieldIteratorFlags::ExcludeSuper); PropIt; ++PropIt)
+		{
+			UProperty* Property = *PropIt;
+			if (!Property->Identical_InContainer(Template, DefObj))
+			{
+				Property->CopyCompleteValue_InContainer(Obj, Template);
+			}
+		}
+	}
+
+	MyDetailLayout->ForceRefreshDetails();
+}
+
+
+//-------------------------------------------------------------------------------------
+
+
+#undef LOCTEXT_NAMESPACE
\ No newline at end of file
diff --git a/Engine/Source/Editor/GameplayAbilitiesEditor/Private/GameplayEffectDetails.h b/Engine/Source/Editor/GameplayAbilitiesEditor/Private/GameplayEffectDetails.h
new file mode 100644
--- /dev/null
+++ b/Engine/Source/Editor/GameplayAbilitiesEditor/Private/GameplayEffectDetails.h
@@ -0,0 +1,41 @@
+// Copyright 1998-2013 Epic Games, Inc. All Rights Reserved.
+
+#pragma once
+
+#include "Editor/PropertyEditor/Public/PropertyEditorModule.h"
+
+class IPropertyHandle;
+class FDetailWidgetRow;
+class IDetailLayoutBuilder;
+
+DECLARE_LOG_CATEGORY_EXTERN(LogGameplayEffectDetails, Log, All);
+
+class UGameplayEffectTemplate;
+
+class FGameplayEffectDetails : public IDetailCustomization
+{
+public:
+	/** Makes a new instance of this detail layout class for a specific detail view requesting it */
+	static TSharedRef<IDetailCustomization> MakeInstance();
+
+private:
+	/** IDetailCustomization interface */
+	virtual void CustomizeDetails(IDetailLayoutBuilder& DetailLayout) override;
+
+	
+
+	TArray<TSharedPtr<FString>> PropertyOptions;
+
+	TSharedPtr<IPropertyHandle> MyProperty;
+
+	IDetailLayoutBuilder* MyDetailLayout;
+
+	TSharedPtr<IPropertyHandle> TemplateProperty;
+	TSharedPtr<IPropertyHandle> ShowAllProperty;
+
+	void OnShowAllChange();
+	void OnTemplateChange();
+
+	bool HideProperties(IDetailLayoutBuilder& DetailLayout, TSharedPtr<IPropertyHandle> PropHandle, UGameplayEffectTemplate* Template);
+};
+
diff --git a/Engine/Source/Runtime/GameplayAbilities/Private/GameplayEffect.cpp b/Engine/Source/Runtime/GameplayAbilities/Private/GameplayEffect.cpp
--- a/Engine/Source/Runtime/GameplayAbilities/Private/GameplayEffect.cpp
+++ b/Engine/Source/Runtime/GameplayAbilities/Private/GameplayEffect.cpp
@@ -30,6 +30,12 @@ UGameplayEffect::UGameplayEffect(const class FPostConstructInitializeProperties&
 	ChanceToExecuteOnGameplayEffect.SetValue(1.f);
 	StackingPolicy = EGameplayEffectStackingPolicy::Unlimited;
 	StackedAttribName = NAME_None;
+
+#if WITH_EDITORONLY_DATA
+	ShowAllProperties = true;
+	Template = nullptr;
+#endif
+
 }
 
 void UGameplayEffect::GetOwnedGameplayTags(FGameplayTagContainer& TagContainer) const
diff --git a/Engine/Source/Runtime/GameplayAbilities/Private/GameplayEffectTemplate.cpp b/Engine/Source/Runtime/GameplayAbilities/Private/GameplayEffectTemplate.cpp
new file mode 100644
--- /dev/null
+++ b/Engine/Source/Runtime/GameplayAbilities/Private/GameplayEffectTemplate.cpp
@@ -0,0 +1,10 @@
+// Copyright 1998-2014 Epic Games, Inc. All Rights Reserved.
+
+#include "AbilitySystemPrivatePCH.h"
+#include "GameplayEffectTemplate.h"
+
+UGameplayEffectTemplate::UGameplayEffectTemplate(const class FPostConstructInitializeProperties& PCIP)
+: Super(PCIP)
+{
+	
+}
\ No newline at end of file
diff --git a/Engine/Source/Runtime/GameplayAbilities/Public/GameplayEffect.h b/Engine/Source/Runtime/GameplayAbilities/Public/GameplayEffect.h
--- a/Engine/Source/Runtime/GameplayAbilities/Public/GameplayEffect.h
+++ b/Engine/Source/Runtime/GameplayAbilities/Public/GameplayEffect.h
@@ -12,6 +12,7 @@
 struct FActiveGameplayEffect;
 
 class UGameplayEffect;
+class UGameplayEffectTemplate;
 class UAbilitySystemComponent;
 
 
@@ -176,9 +177,10 @@ struct FGameplayEffectCue
 UCLASS(BlueprintType)
 class GAMEPLAYABILITIES_API UGameplayEffect : public UDataAsset, public IGameplayTagAssetInterface
 {
-	GENERATED_UCLASS_BODY()
 
 public:
+	GENERATED_UCLASS_BODY()
+
 	/** Infinite duration */
 	static const float INFINITE_DURATION;
 
@@ -187,6 +189,17 @@ public:
 
 	/** Constant specifying that the combat effect has no period and doesn't check for over time application */
 	static const float NO_PERIOD;
+
+#if WITH_EDITORONLY_DATA
+	/** Template to derive starting values and editing customization from */
+	UPROPERTY(EditDefaultsOnly, BlueprintReadOnly, Category = Template)
+	UGameplayEffectTemplate*	Template;
+
+	/** When false, show a limited set of properties for editing, based on the template we are derived from */
+	UPROPERTY(EditDefaultsOnly, BlueprintReadOnly, Category = Template)
+	bool ShowAllProperties;
+#endif
+
 	
 	UPROPERTY(EditDefaultsOnly, Category=GameplayEffect)
 	FScalableFloat	Duration;
@@ -326,8 +339,12 @@ public:
 	UPROPERTY(EditDefaultsOnly, Category = Stacking)
 	TSubclassOf<class UGameplayEffectStackingExtension> StackingExtension;
 
+
+
 protected:
 	void ValidateStacking();
+
+
 };
 
 /**
diff --git a/Engine/Source/Runtime/GameplayAbilities/Public/GameplayEffectTemplate.h b/Engine/Source/Runtime/GameplayAbilities/Public/GameplayEffectTemplate.h
new file mode 100644
--- /dev/null
+++ b/Engine/Source/Runtime/GameplayAbilities/Public/GameplayEffectTemplate.h
@@ -0,0 +1,37 @@
+// Copyright 1998-2014 Epic Games, Inc. All Rights Reserved.
+
+#pragma once
+#include "GameplayEffect.h"
+#include "GameplayEffectTemplate.generated.h"
+
+/**
+ *	The idea here is that UGameplayEffectTemplates are templates created by system designers for commonly used
+ *	GameplayEffects. For example: StandardPhysicalDamage, StandardSnare, etc. 
+ *	
+ *	This allows for A) Less experienced designers to create new ones based on the templates B) Easier/quicker creation
+ *	of GameplayEffects: don't have to worry about forgeting a tag or stacking rule, etc.
+ *	
+ *	The template provides: 
+ *		-Default values
+ *		-What properties are editable by default
+ *		
+ *	Default values are simply the default values of the UGameplayEffectTemplate (since it is a UGameplayEffect). Whatever values
+ *	are different than CDO values will be copied over when making a new GameplayEffect from a template.
+ *	
+ *	EditableProperties describes what properties should be exposed by default when editing a templated gameplayeffect. This is just
+ *	a string list. Properties whose name match strings in that list will be shown, the rest will be hidden. All properties can always be
+ *	seen/edited by just checking "ShowAllProperties" on the GameplayEffect itself. 
+ *	
+ *	
+ */
+UCLASS()
+class UGameplayEffectTemplate : public UGameplayEffect
+{
+	GENERATED_UCLASS_BODY()
+
+#if WITH_EDITORONLY_DATA
+	/** Names of the properties that should be exposed in the limited/basic view when editing a GameplayEffect based on this template */
+	UPROPERTY(EditDefaultsOnly, Category="Template")
+	TArray<FString>	EditableProperties;
+#endif
+};
\ No newline at end of file
