small fix

diff --git a/Engine/Source/Editor/GameplayAbilitiesEditor/Private/GameplayEffectDetails.cpp b/Engine/Source/Editor/GameplayAbilitiesEditor/Private/GameplayEffectDetails.cpp
--- a/Engine/Source/Editor/GameplayAbilitiesEditor/Private/GameplayEffectDetails.cpp
+++ b/Engine/Source/Editor/GameplayAbilitiesEditor/Private/GameplayEffectDetails.cpp
@@ -69,13 +69,13 @@ END_SLATE_FUNCTION_BUILD_OPTIMIZATION
 /** Recursively hide properties that are not default editable according to the Template */
 bool FGameplayEffectDetails::HideProperties(IDetailLayoutBuilder& DetailLayout, TSharedPtr<IPropertyHandle> PropHandle, UGameplayEffectTemplate* Template)
 {
+	UProperty* UProp = PropHandle->GetProperty();
+
 	// Never hide the Template or ShowAllProperties properties
-	if (TemplateProperty == PropHandle || ShowAllProperty == PropHandle)
+	if (TemplateProperty->GetProperty() == UProp || ShowAllProperty->GetProperty() == UProp)
 	{
 		return false;
 	}
-	
-	UProperty* UProp = PropHandle->GetProperty();
 
 	// Don't hide the EditableProperties
 	for (FString MatchStr : Template->EditableProperties)
@@ -134,6 +134,9 @@ void FGameplayEffectDetails::OnTemplateChange()
 			{
 				Property->CopyCompleteValue_InContainer(Obj, Template);
 			}
+
+			// Default to only showing template properties after changing template type
+			Obj->ShowAllProperties = false;
 		}
 	}
 
