missed files

diff --git a/Engine/Source/Editor/GameplayAbilitiesEditor/Private/AttributeDetails.cpp b/Engine/Source/Editor/GameplayAbilitiesEditor/Private/AttributeDetails.cpp
--- a/Engine/Source/Editor/GameplayAbilitiesEditor/Private/AttributeDetails.cpp
+++ b/Engine/Source/Editor/GameplayAbilitiesEditor/Private/AttributeDetails.cpp
@@ -3,7 +3,6 @@
 #include "AbilitySystemEditorPrivatePCH.h"
 #include "AttributeDetails.h"
 #include "Editor/PropertyEditor/Public/PropertyEditing.h"
-#include "AttributeDetails.h"
 #include "AttributeSet.h"
 #include "BlueprintUtilities.h"
 #include "KismetEditorUtilities.h"
diff --git a/Engine/Source/Editor/GameplayAbilitiesEditor/Private/GameplayAbilitiesEditorModule.cpp b/Engine/Source/Editor/GameplayAbilitiesEditor/Private/GameplayAbilitiesEditorModule.cpp
--- a/Engine/Source/Editor/GameplayAbilitiesEditor/Private/GameplayAbilitiesEditorModule.cpp
+++ b/Engine/Source/Editor/GameplayAbilitiesEditor/Private/GameplayAbilitiesEditorModule.cpp
@@ -4,6 +4,7 @@
 #include "Editor/PropertyEditor/Public/PropertyEditorModule.h"
 #include "AttributeDetails.h"
 #include "AttributeSet.h"
+#include "GameplayEffectDetails.h"
 
 #include "IAssetTypeActions.h"
 #include "AssetToolsModule.h"
@@ -36,6 +37,7 @@ void FGameplayAbilitiesEditorModule::StartupModule()
 	PropertyModule.RegisterCustomPropertyTypeLayout( "ScalableFloat", FOnGetPropertyTypeCustomizationInstance::CreateStatic( &FScalableFloatDetails::MakeInstance ) );
 
 	PropertyModule.RegisterCustomClassLayout( "AttributeSet", FOnGetDetailCustomizationInstance::CreateStatic( &FAttributeDetails::MakeInstance ) );
+	PropertyModule.RegisterCustomClassLayout( "GameplayEffect", FOnGetDetailCustomizationInstance::CreateStatic( &FGameplayEffectDetails::MakeInstance ) );
 
 	// Register asset types
 	IAssetTools& AssetTools = FModuleManager::LoadModuleChecked<FAssetToolsModule>("AssetTools").Get();
diff --git a/Engine/Source/Runtime/GameplayAbilities/Public/Abilities/GameplayAbilityTargetActor_SingleLineTrace.h b/Engine/Source/Runtime/GameplayAbilities/Public/Abilities/GameplayAbilityTargetActor_SingleLineTrace.h
--- a/Engine/Source/Runtime/GameplayAbilities/Public/Abilities/GameplayAbilityTargetActor_SingleLineTrace.h
+++ b/Engine/Source/Runtime/GameplayAbilities/Public/Abilities/GameplayAbilityTargetActor_SingleLineTrace.h
@@ -32,7 +32,4 @@ public:
 
 	UPROPERTY(BlueprintReadOnly, meta=(ExposeOnSpawn=true), Category=Projectile)
 	bool bBindToConfirmCancelInputs;
-
-	UPROPERTY(BlueprintReadOnly, meta = (ExposeOnSpawn = true), Category = Projectile)
-	FVector StartingOffset;
 };
\ No newline at end of file
