Show default values in ExposeOnSpawn parameters in actor spawning ability tasks

diff --git a/Engine/Source/Editor/GameplayAbilitiesEditor/Private/K2Node_LatentAbilityCall.cpp b/Engine/Source/Editor/GameplayAbilitiesEditor/Private/K2Node_LatentAbilityCall.cpp
--- a/Engine/Source/Editor/GameplayAbilitiesEditor/Private/K2Node_LatentAbilityCall.cpp
+++ b/Engine/Source/Editor/GameplayAbilitiesEditor/Private/K2Node_LatentAbilityCall.cpp
@@ -233,13 +233,22 @@ void UK2Node_LatentAbilityCall::CreatePinsForClass(UClass* InClass)
 			bIsSettableExternally &&
 			Property->HasAllPropertyFlags(CPF_BlueprintVisible) &&
 			!bIsDelegate && 
-			!IgnorePropertyList.Contains(Property->GetName()))
+			!IgnorePropertyList.Contains(Property->GetName()) &&
+			(FindPin(Property->GetName()) == nullptr) )
 		{
 
 
 			UEdGraphPin* Pin = CreatePin(EGPD_Input, TEXT(""), TEXT(""), NULL, false, false, Property->GetName());
 			const bool bPinGood = (Pin != NULL) && K2Schema->ConvertPropertyToPinType(Property, /*out*/ Pin->PinType);
 			SpawnParmPins.Add(Pin);
+
+			if (K2Schema->PinDefaultValueIsEditable(*Pin))
+			{
+				FString DefaultValueAsString;
+				const bool bDefaultValueSet = FBlueprintEditorUtils::PropertyValueToString(Property, (uint8*)InClass->ClassDefaultObject, DefaultValueAsString);
+				check(bDefaultValueSet);
+				K2Schema->TrySetDefaultValue(*Pin, DefaultValueAsString);
+			}
 		}
 	}
 }
@@ -506,12 +515,34 @@ void UK2Node_LatentAbilityCall::ExpandNode(class FKismetCompilerContext& Compile
 
 	// -------------------------------------------
 	// Set spawn variables
+	//  Borrowed heavily from FKismetCompilerUtilities::GenerateAssignmentNodes
 	// -------------------------------------------
 	
 	for (auto SpawnVarPin : SpawnParmPins)
 	{
 		if (SpawnVarPin->LinkedTo.Num() > 0 || SpawnVarPin->DefaultValue != FString())
 		{
+			if (SpawnVarPin->LinkedTo.Num() == 0)
+			{
+				UProperty* Property = FindField<UProperty>(ClassToSpawn, *SpawnVarPin->PinName);
+				// NULL property indicates that this pin was part of the original node, not the 
+				// class we're assigning to:
+				if (!Property)
+				{
+					continue;
+				}
+
+				// We don't want to generate an assignment node unless the default value 
+				// differs from the value in the CDO:
+				FString DefaultValueAsString;
+				FBlueprintEditorUtils::PropertyValueToString(Property, (uint8*)ClassToSpawn->ClassDefaultObject, DefaultValueAsString);
+				if (DefaultValueAsString == SpawnVarPin->DefaultValue)
+				{
+					continue;
+				}
+			}
+
+
 			UFunction* SetByNameFunction = Schema->FindSetVariableByNameFunction(SpawnVarPin->PinType);
 			if (SetByNameFunction)
 			{
