UE4: Non-unity fix CIS

diff --git a/Engine/Source/Editor/GameplayAbilitiesEditor/Private/K2Node_GameplayCueEvent.cpp b/Engine/Source/Editor/GameplayAbilitiesEditor/Private/K2Node_GameplayCueEvent.cpp
--- a/Engine/Source/Editor/GameplayAbilitiesEditor/Private/K2Node_GameplayCueEvent.cpp
+++ b/Engine/Source/Editor/GameplayAbilitiesEditor/Private/K2Node_GameplayCueEvent.cpp
@@ -8,6 +8,7 @@
 #include "GameplayTagsModule.h"
 #include "GameplayCueInterface.h"
 #include "BlueprintEventNodeSpawner.h"
+#include "BlueprintEditorUtils.h"
 
 #define LOCTEXT_NAMESPACE "K2Node_GameplayCueEvent"
 
