Update asset type actions for GameplayAbilitiesBlueprint to inherit from Blueprint and not Base.

Fixes a bug where diffs of GameplayAbilitiesBlueprints use a text diff instead of the visual diff used by blueprints.

diff --git a/Engine/Source/Developer/AssetTools/Private/AssetTypeActions/AssetTypeActions_Blueprint.h b/Engine/Source/Developer/AssetTools/Private/AssetTypeActions/AssetTypeActions_Blueprint.h
--- a/Engine/Source/Developer/AssetTools/Private/AssetTypeActions/AssetTypeActions_Blueprint.h
+++ b/Engine/Source/Developer/AssetTools/Private/AssetTypeActions/AssetTypeActions_Blueprint.h
@@ -2,7 +2,9 @@
 
 #pragma once
 
-class FAssetTypeActions_Blueprint : public FAssetTypeActions_Base
+#include "AssetTypeActions_Base.h"
+
+class ASSETTOOLS_API FAssetTypeActions_Blueprint : public FAssetTypeActions_Base
 {
 public:
 	// IAssetTypeActions Implementation
diff --git a/Engine/Source/Editor/GameplayAbilitiesEditor/GameplayAbilitiesEditor.Build.cs b/Engine/Source/Editor/GameplayAbilitiesEditor/GameplayAbilitiesEditor.Build.cs
--- a/Engine/Source/Editor/GameplayAbilitiesEditor/GameplayAbilitiesEditor.Build.cs
+++ b/Engine/Source/Editor/GameplayAbilitiesEditor/GameplayAbilitiesEditor.Build.cs
@@ -11,6 +11,7 @@ namespace UnrealBuildTool.Rules
                     "Editor/GraphEditor/Private",
 				    "Editor/Kismet/Private",
 					"Editor/GameplayAbilitiesEditor/Private",
+                    "Developer/AssetTools/Private",
 				}
 			);
 
diff --git a/Engine/Source/Editor/GameplayAbilitiesEditor/Private/AssetTypeActions_GameplayAbilitiesBlueprint.h b/Engine/Source/Editor/GameplayAbilitiesEditor/Private/AssetTypeActions_GameplayAbilitiesBlueprint.h
--- a/Engine/Source/Editor/GameplayAbilitiesEditor/Private/AssetTypeActions_GameplayAbilitiesBlueprint.h
+++ b/Engine/Source/Editor/GameplayAbilitiesEditor/Private/AssetTypeActions_GameplayAbilitiesBlueprint.h
@@ -2,10 +2,9 @@
 
 #pragma once
 
-#include "AssetTypeActions_Base.h"
+#include "AssetTypeActions/AssetTypeActions_Blueprint.h"
 
-//TODO Use base class FAssetTypeActions_Blueprint later
-class FAssetTypeActions_GameplayAbilitiesBlueprint : public FAssetTypeActions_Base
+class FAssetTypeActions_GameplayAbilitiesBlueprint : public FAssetTypeActions_Blueprint
 {
 public:
 	// IAssetTypeActions Implementation
