Fix gameplay cue event enumeration

diff --git a/Engine/Source/Editor/GameplayAbilitiesEditor/Private/K2Node_GameplayCueEvent.cpp b/Engine/Source/Editor/GameplayAbilitiesEditor/Private/K2Node_GameplayCueEvent.cpp
--- a/Engine/Source/Editor/GameplayAbilitiesEditor/Private/K2Node_GameplayCueEvent.cpp
+++ b/Engine/Source/Editor/GameplayAbilitiesEditor/Private/K2Node_GameplayCueEvent.cpp
@@ -34,15 +34,11 @@ FText UK2Node_GameplayCueEvent::GetNodeTitle(ENodeTitleType::Type TitleType) con
 bool UK2Node_GameplayCueEvent::IsCompatibleWithGraph(UEdGraph const* TargetGraph) const
 {
 	bool bIsCompatible = false;
-	if (bIsCompatible)
+	if (UBlueprint* Blueprint = FBlueprintEditorUtils::FindBlueprintForGraph(TargetGraph))
 	{
-		if (UBlueprint* Blueprint = FBlueprintEditorUtils::FindBlueprintForGraph(TargetGraph))
-		{
-			check(Blueprint->GeneratedClass != nullptr);
-			bIsCompatible = Blueprint->GeneratedClass->ImplementsInterface(UGameplayCueInterface::StaticClass());
-		}
-	}
-	
+		check(Blueprint->GeneratedClass != nullptr);
+		bIsCompatible = Blueprint->GeneratedClass->ImplementsInterface(UGameplayCueInterface::StaticClass());
+	}	
 	return bIsCompatible && Super::IsCompatibleWithGraph(TargetGraph);
 }
 
