hide test attributes from details customization property lists

diff --git a/Engine/Source/Editor/GameplayAbilitiesEditor/Private/AttributeDetails.cpp b/Engine/Source/Editor/GameplayAbilitiesEditor/Private/AttributeDetails.cpp
--- a/Engine/Source/Editor/GameplayAbilitiesEditor/Private/AttributeDetails.cpp
+++ b/Engine/Source/Editor/GameplayAbilitiesEditor/Private/AttributeDetails.cpp
@@ -32,12 +32,17 @@ void FAttributePropertyDetails::CustomizeHeader( TSharedRef<IPropertyHandle> Str
 
 	FString FilterMetaStr = StructPropertyHandle->GetProperty()->GetMetaData(TEXT("FilterMetaTag"));
 
-	// Gather all UAttraibute classes
+	// Gather all UAttriubute classes
 	for (TObjectIterator<UClass> ClassIt; ClassIt; ++ClassIt)
 	{
 		UClass *Class = *ClassIt;
 		if (Class->IsChildOf(UAttributeSet::StaticClass()) && !FKismetEditorUtilities::IsClassABlueprintSkeleton(Class))
 		{
+			if (Class->HasMetaData(TEXT("HideInDetailsView")))
+			{
+				continue;
+			}
+
 			for (TFieldIterator<UProperty> PropertyIt(Class, EFieldIteratorFlags::ExcludeSuper); PropertyIt; ++PropertyIt)
 			{
 				UProperty *Property = *PropertyIt;
diff --git a/Engine/Source/Runtime/GameplayAbilities/Public/AbilitySystemTestAttributeSet.h b/Engine/Source/Runtime/GameplayAbilities/Public/AbilitySystemTestAttributeSet.h
--- a/Engine/Source/Runtime/GameplayAbilities/Public/AbilitySystemTestAttributeSet.h
+++ b/Engine/Source/Runtime/GameplayAbilities/Public/AbilitySystemTestAttributeSet.h
@@ -4,7 +4,7 @@
 #include "AttributeSet.h"
 #include "AbilitySystemTestAttributeSet.generated.h"
 
-UCLASS(Blueprintable, BlueprintType)
+UCLASS(Blueprintable, BlueprintType, meta=(HideInDetailsView))
 class GAMEPLAYABILITIES_API UAbilitySystemTestAttributeSet : public UAttributeSet
 {
 	GENERATED_UCLASS_BODY()
