Multi-projectile spawning works in Orion, multi-actor spawning functionality has been added to engine/BP stuff.

diff --git a/Engine/Source/Editor/GameplayAbilitiesEditor/Classes/K2Node_LatentAbilityCall.h b/Engine/Source/Editor/GameplayAbilitiesEditor/Classes/K2Node_LatentAbilityCall.h
--- a/Engine/Source/Editor/GameplayAbilitiesEditor/Classes/K2Node_LatentAbilityCall.h
+++ b/Engine/Source/Editor/GameplayAbilitiesEditor/Classes/K2Node_LatentAbilityCall.h
@@ -22,6 +22,7 @@ class UK2Node_LatentAbilityCall : public UK2Node_BaseAsyncTask
 
 	virtual void ReallocatePinsDuringReconstruction(TArray<UEdGraphPin*>& OldPins) override;
 	virtual void PinDefaultValueChanged(UEdGraphPin* Pin) override;
+	bool ConnectSpawnProperties(UClass* ClassToSpawn, const UEdGraphSchema_K2* Schema, class FKismetCompilerContext& CompilerContext, UEdGraph* SourceGraph, UEdGraphPin*& LastThenPin, UEdGraphPin* SpawnedActorReturnPin);		//Helper
 	virtual void ExpandNode(class FKismetCompilerContext& CompilerContext, UEdGraph* SourceGraph) override;
 
 	void CreatePinsForClass(UClass* InClass);
@@ -29,7 +30,8 @@ class UK2Node_LatentAbilityCall : public UK2Node_BaseAsyncTask
 	UClass* GetClassToSpawn(const TArray<UEdGraphPin*>* InPinsToSearch = NULL) const;
 	UEdGraphPin* GetResultPin() const;
 	bool IsSpawnVarPin(UEdGraphPin* Pin);
-	bool ValidateActorSpawning(class FKismetCompilerContext& CompilerContext);
+	bool ValidateActorSpawning(class FKismetCompilerContext& CompilerContext, bool bGenerateErrors);
+	bool ValidateActorArraySpawning(class FKismetCompilerContext& CompilerContext, bool bGenerateErrors);
 
 	UPROPERTY()
 	TArray<UEdGraphPin*>	SpawnParmPins;
diff --git a/Engine/Source/Editor/GameplayAbilitiesEditor/Private/K2Node_LatentAbilityCall.cpp b/Engine/Source/Editor/GameplayAbilitiesEditor/Private/K2Node_LatentAbilityCall.cpp
--- a/Engine/Source/Editor/GameplayAbilitiesEditor/Private/K2Node_LatentAbilityCall.cpp
+++ b/Engine/Source/Editor/GameplayAbilitiesEditor/Private/K2Node_LatentAbilityCall.cpp
@@ -110,7 +110,7 @@
 			{
 				continue;
 			}
-			
+
 			UObjectProperty* ReturnProperty = Cast<UObjectProperty>(Function->GetReturnProperty());
 			// see if the function is a static factory method for online proxies
 			bool const bIsProxyFactoryMethod = (ReturnProperty != nullptr) && ReturnProperty->PropertyClass->IsChildOf<UAbilityTask>();
@@ -154,6 +154,8 @@
 	static FString ClassPinName;
 	static FString BeginSpawnFuncName;
 	static FString FinishSpawnFuncName;
+	static FString BeginSpawnArrayFuncName;
+	static FString FinishSpawnArrayFuncName;
 	static FString SpawnedActorPinName;
 };
 
@@ -161,6 +163,8 @@
 FString FK2Node_LatentAbilityCallHelper::ClassPinName(TEXT("Class"));
 FString FK2Node_LatentAbilityCallHelper::BeginSpawnFuncName(TEXT("BeginSpawningActor"));
 FString FK2Node_LatentAbilityCallHelper::FinishSpawnFuncName(TEXT("FinishSpawningActor"));
+FString FK2Node_LatentAbilityCallHelper::BeginSpawnArrayFuncName(TEXT("BeginSpawningActorArray"));
+FString FK2Node_LatentAbilityCallHelper::FinishSpawnArrayFuncName(TEXT("FinishSpawningActorArray"));
 FString FK2Node_LatentAbilityCallHelper::SpawnedActorPinName(TEXT("SpawnedActor"));
 
 // -------------------------------------------------
@@ -318,7 +322,7 @@
 
 }
 
-bool UK2Node_LatentAbilityCall::ValidateActorSpawning(class FKismetCompilerContext& CompilerContext)
+bool UK2Node_LatentAbilityCall::ValidateActorSpawning(class FKismetCompilerContext& CompilerContext, bool bGenerateErrors)
 {
 	FName ProxyPrespawnFunctionName = *FK2Node_LatentAbilityCallHelper::BeginSpawnFuncName;
 	UFunction* PreSpawnFunction = ProxyFactoryClass->FindFunctionByName(ProxyPrespawnFunctionName);
@@ -326,33 +330,204 @@
 	FName ProxyPostpawnFunctionName = *FK2Node_LatentAbilityCallHelper::FinishSpawnFuncName;
 	UFunction* PostSpawnFunction = ProxyFactoryClass->FindFunctionByName(ProxyPostpawnFunctionName);
 
+	FName ProxyPrespawnArrayFunctionName = *FK2Node_LatentAbilityCallHelper::BeginSpawnArrayFuncName;
+	UFunction* PreSpawnArrayFunction = ProxyFactoryClass->FindFunctionByName(ProxyPrespawnArrayFunctionName);
+
+	FName ProxyPostpawnArrayFunctionName = *FK2Node_LatentAbilityCallHelper::FinishSpawnArrayFuncName;
+	UFunction* PostSpawnArrayFunction = ProxyFactoryClass->FindFunctionByName(ProxyPostpawnArrayFunctionName);
+
 	bool HasClassParameter = GetClassToSpawn() != nullptr;
 	bool HasPreSpawnFunc = PreSpawnFunction != nullptr;
 	bool HasPostSpawnFunc = PostSpawnFunction != nullptr;
+	bool HasPreSpawnArrayFunc = PreSpawnArrayFunction != nullptr;
+	bool HasPostSpawnArrayFunc = PostSpawnArrayFunction != nullptr;
 
-	
 	if (HasClassParameter || HasPreSpawnFunc || HasPostSpawnFunc)
 	{
 		// They are trying to use ActorSpawning. If any of the above are NOT true, then we have a problem
 		if (!HasClassParameter)
 		{
-			CompilerContext.MessageLog.Error(*LOCTEXT("MissingClassParameter", "UK2Node_LatentAbilityCall: Attempting to use ActorSpawning but Proxy Factory Function missing a Class parameter. @@").ToString(), this);
+			if (bGenerateErrors)
+			{
+				CompilerContext.MessageLog.Error(*LOCTEXT("MissingClassParameter", "UK2Node_LatentAbilityCall: Attempting to use ActorSpawning but Proxy Factory Function missing a Class parameter. @@").ToString(), this);
+			}
 			return false;
 		}
 		if (!HasPreSpawnFunc)
 		{
-			CompilerContext.MessageLog.Error(*LOCTEXT("MissingBeginSpawningFunc", "UK2Node_LatentAbilityCall: Attempting to use ActorSpawning but Missing a BeginSpawningActor function. @@").ToString(), this);
+			if (bGenerateErrors)
+			{
+				CompilerContext.MessageLog.Error(*LOCTEXT("MissingBeginSpawningFunc", "UK2Node_LatentAbilityCall: Attempting to use ActorSpawning but Missing a BeginSpawningActor function. @@").ToString(), this);
+			}
 			return false;
 		}
 		if (!HasPostSpawnFunc)
 		{
-			CompilerContext.MessageLog.Error(*LOCTEXT("MissingFinishSpawningFunc", "UK2Node_LatentAbilityCall: Attempting to use ActorSpawning but Missing a FinishSpawningActor function. @@").ToString(), this);
+			if (bGenerateErrors)
+			{
+				CompilerContext.MessageLog.Error(*LOCTEXT("MissingFinishSpawningFunc", "UK2Node_LatentAbilityCall: Attempting to use ActorSpawning but Missing a FinishSpawningActor function. @@").ToString(), this);
+			}
 			return false;
 		}
-	}	
+		if ((HasPreSpawnArrayFunc || HasPostSpawnArrayFunc))
+		{
+			if (bGenerateErrors)
+			{
+				CompilerContext.MessageLog.Error(*LOCTEXT("SpawnFuncAmbiguous", "UK2Node_LatentAbilityCall: Both ActorSpawning and ActorArraySpawning are at least partially implemented. These are mutually exclusive. @@").ToString(), this);
+			}
+			return false;
+		}
+	}
 	
 	return true;
-	
 }
 
+bool UK2Node_LatentAbilityCall::ValidateActorArraySpawning(class FKismetCompilerContext& CompilerContext, bool bGenerateErrors)
+{
+	FName ProxyPrespawnFunctionName = *FK2Node_LatentAbilityCallHelper::BeginSpawnFuncName;
+	UFunction* PreSpawnFunction = ProxyFactoryClass->FindFunctionByName(ProxyPrespawnFunctionName);
+
+	FName ProxyPostpawnFunctionName = *FK2Node_LatentAbilityCallHelper::FinishSpawnFuncName;
+	UFunction* PostSpawnFunction = ProxyFactoryClass->FindFunctionByName(ProxyPostpawnFunctionName);
+
+	FName ProxyPrespawnArrayFunctionName = *FK2Node_LatentAbilityCallHelper::BeginSpawnArrayFuncName;
+	UFunction* PreSpawnArrayFunction = ProxyFactoryClass->FindFunctionByName(ProxyPrespawnArrayFunctionName);
+
+	FName ProxyPostpawnArrayFunctionName = *FK2Node_LatentAbilityCallHelper::FinishSpawnArrayFuncName;
+	UFunction* PostSpawnArrayFunction = ProxyFactoryClass->FindFunctionByName(ProxyPostpawnArrayFunctionName);
+
+	bool HasClassParameter = GetClassToSpawn() != nullptr;
+	bool HasPreSpawnFunc = PreSpawnFunction != nullptr;
+	bool HasPostSpawnFunc = PostSpawnFunction != nullptr;
+	bool HasPreSpawnArrayFunc = PreSpawnArrayFunction != nullptr;
+	bool HasPostSpawnArrayFunc = PostSpawnArrayFunction != nullptr;
+
+	if (HasClassParameter || HasPreSpawnFunc || HasPostSpawnFunc || HasPreSpawnArrayFunc || HasPostSpawnArrayFunc)
+	{
+		// They are trying to use ActorSpawning. If any of the above are NOT true, then we have a problem
+		if (!HasClassParameter)
+		{
+			if (bGenerateErrors)
+			{
+				CompilerContext.MessageLog.Error(*LOCTEXT("MissingClassParameter", "UK2Node_LatentAbilityCall: Attempting to use ActorSpawning but Proxy Factory Function missing a Class parameter. @@").ToString(), this);
+			}
+			return false;
+		}
+		if (!HasPreSpawnArrayFunc)
+		{
+			if (bGenerateErrors)
+			{
+				CompilerContext.MessageLog.Error(*LOCTEXT("MissingBeginSpawningArrayFunc", "UK2Node_LatentAbilityCall: Attempting to use ActorArraySpawning but Missing a BeginSpawningActorArray function. @@").ToString(), this);
+			}
+			return false;
+		}
+		if (!HasPostSpawnArrayFunc)
+		{
+			if (bGenerateErrors)
+			{
+				CompilerContext.MessageLog.Error(*LOCTEXT("MissingFinishSpawningArrayFunc", "UK2Node_LatentAbilityCall: Attempting to use ActorArraySpawning but Missing a FinishSpawningActorArray function. @@").ToString(), this);
+			}
+			return false;
+		}
+		if (HasPreSpawnFunc || HasPostSpawnFunc)
+		{
+			if (bGenerateErrors)
+			{
+				CompilerContext.MessageLog.Error(*LOCTEXT("SpawnFuncAmbiguous", "UK2Node_LatentAbilityCall: Both ActorSpawning and ActorArraySpawning are at least partially implemented. These are mutually exclusive. @@").ToString(), this);
+			}
+			return false;
+		}
+	}
+
+	return true;
+}
+
+bool UK2Node_LatentAbilityCall::ConnectSpawnProperties(UClass* ClassToSpawn, const UEdGraphSchema_K2* Schema, class FKismetCompilerContext& CompilerContext, UEdGraph* SourceGraph, UEdGraphPin*& LastThenPin, UEdGraphPin* SpawnedActorReturnPin)
+{
+	bool bIsErrorFree = true;
+	for (auto SpawnVarPin : SpawnParmPins)
+	{
+		const bool bHasDefaultValue = !SpawnVarPin->DefaultValue.IsEmpty() || !SpawnVarPin->DefaultTextValue.IsEmpty() || SpawnVarPin->DefaultObject;
+		if (SpawnVarPin->LinkedTo.Num() > 0 || bHasDefaultValue)
+		{
+			if (SpawnVarPin->LinkedTo.Num() == 0)
+			{
+				UProperty* Property = FindField<UProperty>(ClassToSpawn, *SpawnVarPin->PinName);
+				// NULL property indicates that this pin was part of the original node, not the 
+				// class we're assigning to:
+				if (!Property)
+				{
+					continue;
+				}
+
+				// We don't want to generate an assignment node unless the default value 
+				// differs from the value in the CDO:
+				FString DefaultValueAsString;
+				FBlueprintEditorUtils::PropertyValueToString(Property, (uint8*)ClassToSpawn->ClassDefaultObject, DefaultValueAsString);
+				if (DefaultValueAsString == SpawnVarPin->DefaultValue)
+				{
+					continue;
+				}
+			}
+
+
+			UFunction* SetByNameFunction = Schema->FindSetVariableByNameFunction(SpawnVarPin->PinType);
+			if (SetByNameFunction)
+			{
+				UK2Node_CallFunction* SetVarNode = NULL;
+				if (SpawnVarPin->PinType.bIsArray)
+				{
+					SetVarNode = CompilerContext.SpawnIntermediateNode<UK2Node_CallArrayFunction>(this, SourceGraph);
+				}
+				else
+				{
+					SetVarNode = CompilerContext.SpawnIntermediateNode<UK2Node_CallFunction>(this, SourceGraph);
+				}
+				SetVarNode->SetFromFunction(SetByNameFunction);
+				SetVarNode->AllocateDefaultPins();
+
+				// Connect this node into the exec chain
+				bIsErrorFree &= Schema->TryCreateConnection(LastThenPin, SetVarNode->GetExecPin());
+				LastThenPin = SetVarNode->GetThenPin();
+
+				static FString ObjectParamName = FString(TEXT("Object"));
+				static FString ValueParamName = FString(TEXT("Value"));
+				static FString PropertyNameParamName = FString(TEXT("PropertyName"));
+
+				// Connect the new actor to the 'object' pin
+				UEdGraphPin* ObjectPin = SetVarNode->FindPinChecked(ObjectParamName);
+				SpawnedActorReturnPin->MakeLinkTo(ObjectPin);
+
+				// Fill in literal for 'property name' pin - name of pin is property name
+				UEdGraphPin* PropertyNamePin = SetVarNode->FindPinChecked(PropertyNameParamName);
+				PropertyNamePin->DefaultValue = SpawnVarPin->PinName;
+
+				UEdGraphPin* ValuePin = SetVarNode->FindPinChecked(ValueParamName);
+				if (SpawnVarPin->LinkedTo.Num() == 0 &&
+					SpawnVarPin->DefaultValue != FString() &&
+					SpawnVarPin->PinType.PinCategory == Schema->PC_Byte &&
+					SpawnVarPin->PinType.PinSubCategoryObject.IsValid() &&
+					SpawnVarPin->PinType.PinSubCategoryObject->IsA<UEnum>())
+				{
+					// Pin is an enum, we need to alias the enum value to an int:
+					UK2Node_EnumLiteral* EnumLiteralNode = CompilerContext.SpawnIntermediateNode<UK2Node_EnumLiteral>(this, SourceGraph);
+					EnumLiteralNode->Enum = CastChecked<UEnum>(SpawnVarPin->PinType.PinSubCategoryObject.Get());
+					EnumLiteralNode->AllocateDefaultPins();
+					EnumLiteralNode->FindPinChecked(Schema->PN_ReturnValue)->MakeLinkTo(ValuePin);
+
+					UEdGraphPin* InPin = EnumLiteralNode->FindPinChecked(UK2Node_EnumLiteral::GetEnumInputPinName());
+					InPin->DefaultValue = SpawnVarPin->DefaultValue;
+				}
+				else
+				{
+					// Move connection from the variable pin on the spawn node to the 'value' pin
+					CompilerContext.MovePinLinksToIntermediate(*SpawnVarPin, *ValuePin);
+					SetVarNode->PinConnectionListChanged(ValuePin);
+				}
+			}
+		}
+	}
+	return bIsErrorFree;
+}
+
 /**
@@ -375,7 +550,8 @@
  */
 void UK2Node_LatentAbilityCall::ExpandNode(class FKismetCompilerContext& CompilerContext, UEdGraph* SourceGraph)
 {
-	ValidateActorSpawning(CompilerContext);
+	bool validatedActorSpawn = ValidateActorSpawning(CompilerContext, false);
+	bool validatedActorArraySpawn = ValidateActorArraySpawning(CompilerContext, false);
 
 	UClass* ClassToSpawn = GetClassToSpawn();
 	if (ClassToSpawn == nullptr)
@@ -392,10 +568,17 @@
 		return;
 	}
 
+	if (!validatedActorSpawn && !validatedActorArraySpawn)
+	{
+		ValidateActorSpawning(CompilerContext, true);
+		ValidateActorArraySpawning(CompilerContext, true);
+	}
+
 	const UEdGraphSchema_K2* Schema = CompilerContext.GetSchema();
 	check(SourceGraph && Schema);
 	bool bIsErrorFree = true;
 
+
 	// ------------------------------------------------------------------------------------------
 	// CREATE A CALL TO FACTORY THE PROXY OBJECT
 	// ------------------------------------------------------------------------------------------
@@ -452,22 +635,36 @@
 	// ------------------------------------------------------------------------------------------
 	// NEW: CREATE A CALL TO THE PRESPAWN FUNCTION, IF IT RETURNS TRUE, THEN WE WILL SPAWN THE NEW ACTOR
 	// ------------------------------------------------------------------------------------------
-	
-	FName ProxyPrespawnFunctionName = *FK2Node_LatentAbilityCallHelper::BeginSpawnFuncName;
+
+	FName ProxyPrespawnFunctionName = validatedActorArraySpawn ? *FK2Node_LatentAbilityCallHelper::BeginSpawnArrayFuncName : *FK2Node_LatentAbilityCallHelper::BeginSpawnFuncName;
 	UFunction* PreSpawnFunction = ProxyFactoryClass->FindFunctionByName(ProxyPrespawnFunctionName);
 
-	FName ProxyPostpawnFunctionName = *FK2Node_LatentAbilityCallHelper::FinishSpawnFuncName;
+	FName ProxyPostpawnFunctionName = validatedActorArraySpawn ? *FK2Node_LatentAbilityCallHelper::FinishSpawnArrayFuncName : *FK2Node_LatentAbilityCallHelper::FinishSpawnFuncName;
 	UFunction* PostSpawnFunction = ProxyFactoryClass->FindFunctionByName(ProxyPostpawnFunctionName);
 
 	if (PreSpawnFunction == nullptr)
 	{
-		CompilerContext.MessageLog.Error(*LOCTEXT("MissingBeginSpawningActorFunction", "AbilityTask: Proxy is missing BeginSpawningActor native function. @@").ToString(), this);
+		if (validatedActorArraySpawn)
+		{
+			CompilerContext.MessageLog.Error(*LOCTEXT("MissingBeginSpawningActorArrayFunction", "AbilityTask: Proxy is missing BeginSpawningActorArray native function. @@").ToString(), this);
+		}
+		else
+		{
+			CompilerContext.MessageLog.Error(*LOCTEXT("MissingBeginSpawningActorArrayFunction", "AbilityTask: Proxy is missing BeginSpawningActorArray native function. @@").ToString(), this);
+		}
 		return;
 	}
 
 	if (PostSpawnFunction == nullptr)
 	{
-		CompilerContext.MessageLog.Error(*LOCTEXT("MissingFinishSpawningActorFunction", "AbilityTask: Proxy is missing FinishSpawningActor native function. @@").ToString(), this);
+		if (validatedActorArraySpawn)
+		{
+			CompilerContext.MessageLog.Error(*LOCTEXT("MissingFinishSpawningActorArrayFunction", "AbilityTask: Proxy is missing FinishSpawningActorArray native function. @@").ToString(), this);
+		}
+		else
+		{
+			CompilerContext.MessageLog.Error(*LOCTEXT("MissingFinishSpawningActorFunction", "AbilityTask: Proxy is missing FinishSpawningActor native function. @@").ToString(), this);
+		}
 		return;
 	}
 
@@ -523,94 +720,88 @@
 
 	LastThenPin = BranchNode->GetThenPin();
 
+	if (validatedActorArraySpawn)
+	{
+		//RICKH Experimental
+
+		//Branch for main loop control
+		UK2Node_IfThenElse* Branch = CompilerContext.SpawnIntermediateNode<UK2Node_IfThenElse>(this, SourceGraph);
+		Branch->AllocateDefaultPins();
+
+		//Create int Iterator
+		UK2Node_TemporaryVariable* IteratorVar = CompilerContext.SpawnIntermediateNode<UK2Node_TemporaryVariable>(this, SourceGraph);
+		IteratorVar->VariableType.PinCategory = Schema->PC_Int;
+		IteratorVar->AllocateDefaultPins();
+
+		//Iterator assignment (initialization to zero)
+		UK2Node_AssignmentStatement* IteratorInitialize = CompilerContext.SpawnIntermediateNode<UK2Node_AssignmentStatement>(this, SourceGraph);
+		IteratorInitialize->AllocateDefaultPins();
+		IteratorInitialize->GetValuePin()->DefaultValue = TEXT("0");
+
+		//Iterator assignment (incrementing)
+		UK2Node_AssignmentStatement* IteratorAssign = CompilerContext.SpawnIntermediateNode<UK2Node_AssignmentStatement>(this, SourceGraph);
+		IteratorAssign->AllocateDefaultPins();
+
+		//Increment iterator command
+		UK2Node_CallFunction* Increment = CompilerContext.SpawnIntermediateNode<UK2Node_CallFunction>(this, SourceGraph);
+		Increment->SetFromFunction(UKismetMathLibrary::StaticClass()->FindFunctionByName(TEXT("Add_IntInt")));
+		Increment->AllocateDefaultPins();
+		Increment->FindPinChecked(TEXT("B"))->DefaultValue = TEXT("1");
+
+		//Array length
+		UK2Node_CallArrayFunction* ArrayLength = CompilerContext.SpawnIntermediateNode<UK2Node_CallArrayFunction>(this, SourceGraph);
+		ArrayLength->SetFromFunction(UKismetArrayLibrary::StaticClass()->FindFunctionByName(TEXT("Array_Length")));
+		ArrayLength->AllocateDefaultPins();
+
+		//Array element retrieval
+		UK2Node_CallArrayFunction* GetElement = CompilerContext.SpawnIntermediateNode<UK2Node_CallArrayFunction>(this, SourceGraph);
+		GetElement->SetFromFunction(UKismetArrayLibrary::StaticClass()->FindFunctionByName(TEXT("Array_Get")));
+		GetElement->AllocateDefaultPins();
+
+		//Check node for iterator versus array length
+		UK2Node_CallFunction* Condition = CompilerContext.SpawnIntermediateNode<UK2Node_CallFunction>(this, SourceGraph);
+		Condition->SetFromFunction(UKismetMathLibrary::StaticClass()->FindFunctionByName(TEXT("Less_IntInt")));
+		Condition->AllocateDefaultPins();
+
+		//Connections to set up the loop
+		UEdGraphSchema_K2 const* const K2Schema = Cast<const UEdGraphSchema_K2>(Schema);
+		bIsErrorFree &= Schema->TryCreateConnection(LastThenPin, IteratorInitialize->GetExecPin());
+		bIsErrorFree &= Schema->TryCreateConnection(IteratorVar->GetVariablePin(), IteratorInitialize->GetVariablePin());
+		bIsErrorFree &= Schema->TryCreateConnection(IteratorInitialize->GetThenPin(), Branch->GetExecPin());
+		bIsErrorFree &= Schema->TryCreateConnection(SpawnedActorReturnPin, ArrayLength->GetTargetArrayPin());
+		bIsErrorFree &= Schema->TryCreateConnection(Condition->GetReturnValuePin(), Branch->GetConditionPin());
+		bIsErrorFree &= Schema->TryCreateConnection(IteratorVar->GetVariablePin(), Condition->FindPinChecked(TEXT("A")));
+		bIsErrorFree &= Schema->TryCreateConnection(ArrayLength->FindPin(K2Schema->PN_ReturnValue), Condition->FindPinChecked(TEXT("B")));
+
+		//Connections to establish loop iteration
+		bIsErrorFree &= Schema->TryCreateConnection(IteratorVar->GetVariablePin(), Increment->FindPinChecked(TEXT("A")));
+		bIsErrorFree &= Schema->TryCreateConnection(IteratorVar->GetVariablePin(), IteratorAssign->GetVariablePin());
+		bIsErrorFree &= Schema->TryCreateConnection(Increment->GetReturnValuePin(), IteratorAssign->GetValuePin());
+		bIsErrorFree &= Schema->TryCreateConnection(IteratorAssign->GetThenPin(), Branch->GetExecPin());
+
+		//This is the inner loop
+		LastThenPin = Branch->GetThenPin();		//Connect the loop branch to the spawn-assignment code block
+		bIsErrorFree &= Schema->TryCreateConnection(SpawnedActorReturnPin, GetElement->GetTargetArrayPin());
+		bIsErrorFree &= Schema->TryCreateConnection(IteratorVar->GetVariablePin(), GetElement->FindPinChecked(K2Schema->PN_Index));
+		UEdGraphPin* RetValPin = GetElement->FindPinChecked(TEXT("Item"));		//TODO Should this be a field? Maybe create UEdGraphSchema_K2::PN_Item?
+		bIsErrorFree &= ConnectSpawnProperties(ClassToSpawn, Schema, CompilerContext, SourceGraph, LastThenPin, RetValPin);		//Last argument needs to be our array element
+		bIsErrorFree &= Schema->TryCreateConnection(LastThenPin, IteratorAssign->GetExecPin());		//Connect the spawn-assignment code block to the iterator increment
+		
+		//Finish by providing the proper path out
+		LastThenPin = Branch->GetElsePin();
+	}
+
 	// -------------------------------------------
 	// Set spawn variables
 	//  Borrowed heavily from FKismetCompilerUtilities::GenerateAssignmentNodes
 	// -------------------------------------------
 	
-	for (auto SpawnVarPin : SpawnParmPins)
+	if (validatedActorSpawn)
 	{
-		const bool bHasDefaultValue = !SpawnVarPin->DefaultValue.IsEmpty() || !SpawnVarPin->DefaultTextValue.IsEmpty() || SpawnVarPin->DefaultObject;
-		if (SpawnVarPin->LinkedTo.Num() > 0 || bHasDefaultValue)
-		{
-			if (SpawnVarPin->LinkedTo.Num() == 0)
-			{
-				UProperty* Property = FindField<UProperty>(ClassToSpawn, *SpawnVarPin->PinName);
-				// NULL property indicates that this pin was part of the original node, not the 
-				// class we're assigning to:
-				if (!Property)
-				{
-					continue;
-				}
-
-				// We don't want to generate an assignment node unless the default value 
-				// differs from the value in the CDO:
-				FString DefaultValueAsString;
-				FBlueprintEditorUtils::PropertyValueToString(Property, (uint8*)ClassToSpawn->ClassDefaultObject, DefaultValueAsString);
-				if (DefaultValueAsString == SpawnVarPin->DefaultValue)
-				{
-					continue;
-				}
-			}
+		bIsErrorFree &= ConnectSpawnProperties(ClassToSpawn, Schema, CompilerContext, SourceGraph, LastThenPin, SpawnedActorReturnPin);
+	}
+	
 
 
-			UFunction* SetByNameFunction = Schema->FindSetVariableByNameFunction(SpawnVarPin->PinType);
-			if (SetByNameFunction)
-			{
-				UK2Node_CallFunction* SetVarNode = NULL;
-				if (SpawnVarPin->PinType.bIsArray)
-				{
-					SetVarNode = CompilerContext.SpawnIntermediateNode<UK2Node_CallArrayFunction>(this, SourceGraph);
-				}
-				else
-				{
-					SetVarNode = CompilerContext.SpawnIntermediateNode<UK2Node_CallFunction>(this, SourceGraph);
-				}
-				SetVarNode->SetFromFunction(SetByNameFunction);
-				SetVarNode->AllocateDefaultPins();
-
-				// Connect this node into the exec chain
-				bIsErrorFree &= Schema->TryCreateConnection(LastThenPin, SetVarNode->GetExecPin());
-				LastThenPin = SetVarNode->GetThenPin();
-
-				static FString ObjectParamName = FString(TEXT("Object"));
-				static FString ValueParamName = FString(TEXT("Value"));
-				static FString PropertyNameParamName = FString(TEXT("PropertyName"));
-
-				// Connect the new actor to the 'object' pin
-				UEdGraphPin* ObjectPin = SetVarNode->FindPinChecked(ObjectParamName);
-				SpawnedActorReturnPin->MakeLinkTo(ObjectPin);
-
-				// Fill in literal for 'property name' pin - name of pin is property name
-				UEdGraphPin* PropertyNamePin = SetVarNode->FindPinChecked(PropertyNameParamName);
-				PropertyNamePin->DefaultValue = SpawnVarPin->PinName;
-				
-				UEdGraphPin* ValuePin = SetVarNode->FindPinChecked(ValueParamName);
-				if (SpawnVarPin->LinkedTo.Num() == 0 &&
-					SpawnVarPin->DefaultValue != FString() &&
-					SpawnVarPin->PinType.PinCategory == Schema->PC_Byte &&
-					SpawnVarPin->PinType.PinSubCategoryObject.IsValid() &&
-					SpawnVarPin->PinType.PinSubCategoryObject->IsA<UEnum>())
-				{
-					// Pin is an enum, we need to alias the enum value to an int:
-					UK2Node_EnumLiteral* EnumLiteralNode = CompilerContext.SpawnIntermediateNode<UK2Node_EnumLiteral>(this, SourceGraph);
-					EnumLiteralNode->Enum = CastChecked<UEnum>(SpawnVarPin->PinType.PinSubCategoryObject.Get());
-					EnumLiteralNode->AllocateDefaultPins();
-					EnumLiteralNode->FindPinChecked(Schema->PN_ReturnValue)->MakeLinkTo(ValuePin);
-
-					UEdGraphPin* InPin = EnumLiteralNode->FindPinChecked(UK2Node_EnumLiteral::GetEnumInputPinName());
-					InPin->DefaultValue = SpawnVarPin->DefaultValue;
-				}
-				else
-				{
-					// Move connection from the variable pin on the spawn node to the 'value' pin
-					CompilerContext.MovePinLinksToIntermediate(*SpawnVarPin, *ValuePin);
-					SetVarNode->PinConnectionListChanged(ValuePin);
-				}
-			}
-		}
-	}	
-		
-
 	// -------------------------------------------
 	// Call FinishSpawning
