Minor cleanup, added a field to Schema for the "Item" output pin. Replaced hardcoded TEXT("Item") in my own code, but left other instances for now.

diff --git a/Engine/Source/Editor/BlueprintGraph/Classes/EdGraphSchema_K2.h b/Engine/Source/Editor/BlueprintGraph/Classes/EdGraphSchema_K2.h
--- a/Engine/Source/Editor/BlueprintGraph/Classes/EdGraphSchema_K2.h
+++ b/Engine/Source/Editor/BlueprintGraph/Classes/EdGraphSchema_K2.h
@@ -256,6 +256,7 @@ class BLUEPRINTGRAPH_API UEdGraphSchema_K2 : public UEdGraphSchema
 	static const FString PN_Start;    // Category=PC_Int, singleton, input
 	static const FString PN_Stop;    // Category=PC_Int, singleton, input
 	static const FString PN_Index;    // Category=PC_Int, singleton, output
+	static const FString PN_Item;    // Category=PC_Int, singleton, output
 	static const FString PN_CastSucceeded;    // Category=PC_Exec, singleton, output
 	static const FString PN_CastFailed;    // Category=PC_Exec, singleton, output
 	static const FString PN_CastedValuePrefix;    // Category=PC_Object, singleton, output; actual pin name varies depending on the type to be casted to, this is just a prefix
diff --git a/Engine/Source/Editor/BlueprintGraph/Private/EdGraphSchema_K2.cpp b/Engine/Source/Editor/BlueprintGraph/Private/EdGraphSchema_K2.cpp
--- a/Engine/Source/Editor/BlueprintGraph/Private/EdGraphSchema_K2.cpp
+++ b/Engine/Source/Editor/BlueprintGraph/Private/EdGraphSchema_K2.cpp
@@ -266,6 +266,7 @@ const FString UEdGraphSchema_K2::PN_Condition(TEXT("Condition"));
 const FString UEdGraphSchema_K2::PN_Start(TEXT("Start"));
 const FString UEdGraphSchema_K2::PN_Stop(TEXT("Stop"));
 const FString UEdGraphSchema_K2::PN_Index(TEXT("Index"));
+const FString UEdGraphSchema_K2::PN_Item(TEXT("Item"));
 const FString UEdGraphSchema_K2::PN_CastSucceeded(TEXT("then"));
 const FString UEdGraphSchema_K2::PN_CastFailed(TEXT("CastFailed"));
 const FString UEdGraphSchema_K2::PN_CastedValuePrefix(TEXT("As"));
diff --git a/Engine/Source/Editor/GameplayAbilitiesEditor/Private/K2Node_LatentAbilityCall.cpp b/Engine/Source/Editor/GameplayAbilitiesEditor/Private/K2Node_LatentAbilityCall.cpp
--- a/Engine/Source/Editor/GameplayAbilitiesEditor/Private/K2Node_LatentAbilityCall.cpp
+++ b/Engine/Source/Editor/GameplayAbilitiesEditor/Private/K2Node_LatentAbilityCall.cpp
@@ -783,8 +783,7 @@ void UK2Node_LatentAbilityCall::ExpandNode(class FKismetCompilerContext& Compile
 		LastThenPin = Branch->GetThenPin();		//Connect the loop branch to the spawn-assignment code block
 		bIsErrorFree &= Schema->TryCreateConnection(SpawnedActorReturnPin, GetElement->GetTargetArrayPin());
 		bIsErrorFree &= Schema->TryCreateConnection(IteratorVar->GetVariablePin(), GetElement->FindPinChecked(K2Schema->PN_Index));
-		UEdGraphPin* RetValPin = GetElement->FindPinChecked(TEXT("Item"));		//TODO Should this be a field? Maybe create UEdGraphSchema_K2::PN_Item?
-		bIsErrorFree &= ConnectSpawnProperties(ClassToSpawn, Schema, CompilerContext, SourceGraph, LastThenPin, RetValPin);		//Last argument needs to be our array element
+		bIsErrorFree &= ConnectSpawnProperties(ClassToSpawn, Schema, CompilerContext, SourceGraph, LastThenPin, GetElement->FindPinChecked(K2Schema->PN_Item));		//Last argument is the array element
 		bIsErrorFree &= Schema->TryCreateConnection(LastThenPin, IteratorAssign->GetExecPin());		//Connect the spawn-assignment code block to the iterator increment
 		
 		//Finish by providing the proper path out
