Added delegate to multi-projectile spawn for each individual projectile, plus a delegate (with NULL projectile) to be called after all projectiles have been spawned. Minor cleanup as well.

diff --git a/Engine/Source/Editor/GameplayAbilitiesEditor/Private/K2Node_LatentAbilityCall.cpp b/Engine/Source/Editor/GameplayAbilitiesEditor/Private/K2Node_LatentAbilityCall.cpp
--- a/Engine/Source/Editor/GameplayAbilitiesEditor/Private/K2Node_LatentAbilityCall.cpp
+++ b/Engine/Source/Editor/GameplayAbilitiesEditor/Private/K2Node_LatentAbilityCall.cpp
@@ -722,8 +722,6 @@ void UK2Node_LatentAbilityCall::ExpandNode(class FKismetCompilerContext& Compile
 
 	if (validatedActorArraySpawn)
 	{
-		//RICKH Experimental
-
 		//Branch for main loop control
 		UK2Node_IfThenElse* Branch = CompilerContext.SpawnIntermediateNode<UK2Node_IfThenElse>(this, SourceGraph);
 		Branch->AllocateDefaultPins();
@@ -800,8 +798,6 @@ void UK2Node_LatentAbilityCall::ExpandNode(class FKismetCompilerContext& Compile
 		bIsErrorFree &= ConnectSpawnProperties(ClassToSpawn, Schema, CompilerContext, SourceGraph, LastThenPin, SpawnedActorReturnPin);
 	}
 	
-
-
 	// -------------------------------------------
 	// Call FinishSpawning
 	// -------------------------------------------
