Fix issue causing gameplay attributes not being resetable to None.

diff --git a/Engine/Source/Editor/GameplayAbilitiesEditor/Private/AttributeDetails.cpp b/Engine/Source/Editor/GameplayAbilitiesEditor/Private/AttributeDetails.cpp
--- a/Engine/Source/Editor/GameplayAbilitiesEditor/Private/AttributeDetails.cpp
+++ b/Engine/Source/Editor/GameplayAbilitiesEditor/Private/AttributeDetails.cpp
@@ -47,10 +47,17 @@ void FAttributePropertyDetails::CustomizeHeader( TSharedRef<IPropertyHandle> Str
 			{
 				UProperty *Property = *PropertyIt;
 
+				// Allow properties to be filtered by specific meta strings
 				if (!FilterMetaStr.IsEmpty() && Property->HasMetaData(*FilterMetaStr))
 				{
 					continue;
 				}
+
+				// Allow properties to be filtered globally (never show up)
+				if (Property->HasMetaData(TEXT("HideInDetailsView")))
+				{
+					continue;
+				}
 				
 				PropertyOptions.Add(MakeShareable(new FString(FString::Printf(TEXT("%s.%s"), *Class->GetName(), *Property->GetName()))));
 			}
@@ -131,13 +138,18 @@ void FAttributePropertyDetails::OnChangeProperty(TSharedPtr<FString> ItemSelecte
 			UProperty *Property = FindField<UProperty>(FoundClass, *PropertyName);
 			if (Property)
 			{
-				const UObject *ObjPtr = Property;
+				const UObject* ObjPtr = Property;
 				MyProperty->SetValue(ObjPtr);
 				
 				return;
 			}
 		}
+
+		const UObject* nullObj = nullptr;
+		MyProperty->SetValue(nullObj);
 	}
+
+	
 }
 
 // ------------------------------------------------------------------------------------
