- first pass on wells and rigs, need to work in the interaction between the 2 but spawning/placing is in!
- moving aim target array to Orion Globals off of game state
- Fix from Ratti for loading BP with K2Node_LatentAbilityCall.cpp
- added a gameplay Set for players where gameplay attributes can go, for now only rig placement range is in there

diff --git a/Engine/Source/Editor/GameplayAbilitiesEditor/Private/K2Node_LatentAbilityCall.cpp b/Engine/Source/Editor/GameplayAbilitiesEditor/Private/K2Node_LatentAbilityCall.cpp
--- a/Engine/Source/Editor/GameplayAbilitiesEditor/Private/K2Node_LatentAbilityCall.cpp
+++ b/Engine/Source/Editor/GameplayAbilitiesEditor/Private/K2Node_LatentAbilityCall.cpp
@@ -219,6 +219,8 @@ void UK2Node_LatentAbilityCall::CreatePinsForClass(UClass* InClass)
 
 	const UEdGraphSchema_K2* K2Schema = GetDefault<UEdGraphSchema_K2>();
 
+	const UObject* const ClassDefaultObject = InClass->GetDefaultObject(false);
+
 	SpawnParmPins.Empty();
 
 	// Tasks can hide spawn parameters by doing meta = (HideSpawnParms="PropertyA,PropertyB")
@@ -258,10 +260,10 @@ void UK2Node_LatentAbilityCall::CreatePinsForClass(UClass* InClass)
 			const bool bPinGood = (Pin != NULL) && K2Schema->ConvertPropertyToPinType(Property, /*out*/ Pin->PinType);
 			SpawnParmPins.Add(Pin);
 
-			if (K2Schema->PinDefaultValueIsEditable(*Pin))
+			if (ClassDefaultObject && K2Schema->PinDefaultValueIsEditable(*Pin))
 			{
 				FString DefaultValueAsString;
-				const bool bDefaultValueSet = FBlueprintEditorUtils::PropertyValueToString(Property, (uint8*)InClass->ClassDefaultObject, DefaultValueAsString);
+				const bool bDefaultValueSet = FBlueprintEditorUtils::PropertyValueToString(Property, reinterpret_cast<const uint8*>(ClassDefaultObject), DefaultValueAsString);
 				check(bDefaultValueSet);
 				K2Schema->TrySetDefaultValue(*Pin, DefaultValueAsString);
 			}
@@ -462,13 +464,16 @@ bool UK2Node_LatentAbilityCall::ConnectSpawnProperties(UClass* ClassToSpawn, con
 					continue;
 				}
 
-				// We don't want to generate an assignment node unless the default value 
-				// differs from the value in the CDO:
-				FString DefaultValueAsString;
-				FBlueprintEditorUtils::PropertyValueToString(Property, (uint8*)ClassToSpawn->ClassDefaultObject, DefaultValueAsString);
-				if (DefaultValueAsString == SpawnVarPin->DefaultValue)
+				if (ClassToSpawn->ClassDefaultObject != nullptr)
 				{
-					continue;
+					// We don't want to generate an assignment node unless the default value 
+					// differs from the value in the CDO:
+					FString DefaultValueAsString;
+					FBlueprintEditorUtils::PropertyValueToString(Property, (uint8*)ClassToSpawn->ClassDefaultObject, DefaultValueAsString);
+					if (DefaultValueAsString == SpawnVarPin->DefaultValue)
+					{
+						continue;
+					}
 				}
 			}
 
