fix filtering code that was stomped in merge

diff --git a/Engine/Source/Editor/GameplayAbilitiesEditor/Private/AttributeDetails.cpp b/Engine/Source/Editor/GameplayAbilitiesEditor/Private/AttributeDetails.cpp
--- a/Engine/Source/Editor/GameplayAbilitiesEditor/Private/AttributeDetails.cpp
+++ b/Engine/Source/Editor/GameplayAbilitiesEditor/Private/AttributeDetails.cpp
@@ -93,6 +93,12 @@ void FAttributePropertyDetails::CustomizeHeader( TSharedRef<IPropertyHandle> Str
 			UClass *Class = *ClassIt;
 			if (Class->IsChildOf(UAttributeSet::StaticClass()) && !FKismetEditorUtilities::IsClassABlueprintSkeleton(Class))
 			{
+				// Allow entire classes to be filtered globalyl
+				if (Class->HasMetaData(TEXT("HideInDetailsView")))
+				{
+					continue;
+				}
+
 				for (TFieldIterator<UProperty> PropertyIt(Class, EFieldIteratorFlags::ExcludeSuper); PropertyIt; ++PropertyIt)
 				{
 					UProperty *Property = *PropertyIt;
diff --git a/Engine/Source/Runtime/GameplayAbilities/Public/GameplayEffect.h b/Engine/Source/Runtime/GameplayAbilities/Public/GameplayEffect.h
--- a/Engine/Source/Runtime/GameplayAbilities/Public/GameplayEffect.h
+++ b/Engine/Source/Runtime/GameplayAbilities/Public/GameplayEffect.h
@@ -110,7 +110,7 @@ struct GAMEPLAYABILITIES_API FGameplayModifierInfo
 	}
 
 	/** The Attribute we modify or the GE we modify modifies. */
-	UPROPERTY(EditDefaultsOnly, Category=GameplayModifier)
+	UPROPERTY(EditDefaultsOnly, Category=GameplayModifier, meta=(FilterMetaTag="HideFromModifiers"))
 	FGameplayAttribute Attribute;
 
 	/** What this modifies - Attribute, OutgoingGEs, IncomingGEs, ACtiveGEs. */
