- Part one of details customization for gameplay effect custom execution scoped modifiers
- Scoped modifiers now show custom combo box showing only valid options derived from execution class

[REDACTED_HANDLE]: Next part: customization on the execution struct to hide the modifier list when there are no valid options

diff --git a/Engine/Source/Editor/GameplayAbilitiesEditor/Private/GameplayEffectExecutionScopedModifierInfoDetails.cpp b/Engine/Source/Editor/GameplayAbilitiesEditor/Private/GameplayEffectExecutionScopedModifierInfoDetails.cpp
--- a/Engine/Source/Editor/GameplayAbilitiesEditor/Private/GameplayEffectExecutionScopedModifierInfoDetails.cpp
+++ b/Engine/Source/Editor/GameplayAbilitiesEditor/Private/GameplayEffectExecutionScopedModifierInfoDetails.cpp
@@ -6,6 +6,8 @@
 #include "IDetailPropertyRow.h"
 #include "PropertyHandle.h"
 #include "GameplayEffect.h"
+#include "GameplayEffectTypes.h"
+#include "GameplayEffectExecutionCalculation.h"
 
 #define LOCTEXT_NAMESPACE "GameplayEffectExecutionScopedModifierInfoDetailsCustomization"
 
@@ -23,9 +25,274 @@ void FGameplayEffectExecutionScopedModifierInfoDetails::CustomizeHeader(TSharedR
 	];
 }
 
+/** Custom widget class to cleanly represent a capture definition in a combo box */
+class SCaptureDefWidget : public SCompoundWidget
+{
+public:
+
+	SLATE_BEGIN_ARGS(SCaptureDefWidget) {}
+	SLATE_END_ARGS()
+
+	/** Construct the widget from a grid panel wrapped with a border */
+	void Construct(const FArguments& InArgs)
+	{
+		this->ChildSlot
+		[
+			SNew(SBorder)
+			.HAlign(HAlign_Fill)
+			[
+				SNew(SGridPanel)
+				+SGridPanel::Slot(0, 0)
+				.HAlign(HAlign_Right)
+				.Padding(FMargin(2.f))
+				[
+					SNew(STextBlock)
+					.Text(NSLOCTEXT("ScopedModifierDetails", "CapturedAttributeLabel", "Captured Attribute:"))
+					.Font(IDetailLayoutBuilder::GetDetailFontBold())
+				]
+				+SGridPanel::Slot(1, 0)
+				.HAlign(HAlign_Left)
+				.Padding(FMargin(2.f))
+				[
+					SNew(STextBlock)
+					.Text(this, &SCaptureDefWidget::GetCapturedAttributeText)
+					.Font(IDetailLayoutBuilder::GetDetailFont())
+				]
+				+SGridPanel::Slot(0, 1)
+				.HAlign(HAlign_Right)
+				.Padding(FMargin(2.f))
+				[
+					SNew(STextBlock)
+					.Text(NSLOCTEXT("ScopedModifierDetails", "CapturedAttributeSourceLabel", "Captured Source:"))
+					.Font(IDetailLayoutBuilder::GetDetailFontBold())
+				]
+				+SGridPanel::Slot(1, 1)
+				.HAlign(HAlign_Left)
+				.Padding(FMargin(2.f))
+				[
+					SNew(STextBlock)
+					.Text(this, &SCaptureDefWidget::GetCapturedAttributeSourceText)
+					.Font(IDetailLayoutBuilder::GetDetailFont())
+				]
+				+SGridPanel::Slot(0, 2)
+				.HAlign(HAlign_Right)
+				.Padding(FMargin(2.f))
+				[
+					SNew(STextBlock)
+					.Text(NSLOCTEXT("ScopedModifierDetails", "CapturedAttributeSnapshotLabel", "Captured Status:"))
+					.Font(IDetailLayoutBuilder::GetDetailFontBold())
+				]
+				+SGridPanel::Slot(1, 2)
+				.HAlign(HAlign_Left)
+				.Padding(FMargin(2.f))
+				[
+					SNew(STextBlock)
+					.Text(this, &SCaptureDefWidget::GetCapturedAttributeSnapshotText)
+					.Font(IDetailLayoutBuilder::GetDetailFont())
+				]
+			]
+		];
+	}
+
+	/** Set the definition that backs the widget; This is used as a way to reduce expensive FText creations constantly */
+	void SetBackingDefinition(const FGameplayEffectAttributeCaptureDefinition& InDefinition)
+	{
+		static const FText SnapshotText(NSLOCTEXT("ScopedModifierDetails", "CapturedAttributeSnapshotted", "Snapshotted"));
+		static const FText NotSnapshotText(NSLOCTEXT("ScopedModifierDetails", "CapturedAttributeNotSnapshotted", "Not Snapshotted"));
+
+		if (InDefinition != BackingDefinition)
+		{
+			BackingDefinition = InDefinition;
+			CapturedAttributeText = FText::FromString(BackingDefinition.AttributeToCapture.GetName());
+			UEnum::GetDisplayValueAsText(TEXT("GameplayAbilities.EGameplayEffectAttributeCaptureSource"), BackingDefinition.AttributeSource, CapturedAttributeSourceText);
+			CapturedAttributeSnapshotText = BackingDefinition.bSnapshot ? SnapshotText : NotSnapshotText;
+		}
+	}
+
+private:
+
+	/** Simple accessor to cached captured attribute text */
+	FText GetCapturedAttributeText() const
+	{
+		return CapturedAttributeText;
+	}
+
+	/** Simple accessor to cached captured attribute source text */
+	FText GetCapturedAttributeSourceText() const
+	{
+		return CapturedAttributeSourceText;
+	}
+
+	/** Simple accessor to cached captured attribute snapshot text */
+	FText GetCapturedAttributeSnapshotText() const
+	{
+		return CapturedAttributeSnapshotText;
+	}
+
+	/** Capture definition backing the widget */
+	FGameplayEffectAttributeCaptureDefinition BackingDefinition;
+
+	/** Cached attribute text */
+	FText CapturedAttributeText;
+
+	/** Cached attribute capture source text */
+	FText CapturedAttributeSourceText;
+
+	/** Cached attribute snapshot status text */
+	FText CapturedAttributeSnapshotText;
+};
+
 void FGameplayEffectExecutionScopedModifierInfoDetails::CustomizeChildren(TSharedRef<IPropertyHandle> StructPropertyHandle, IDetailChildrenBuilder& StructBuilder, IPropertyTypeCustomizationUtils& StructCustomizationUtils)
 {
-	// todo: Implement me!
+	AvailableCaptureDefs.Empty();
+	CaptureDefPropertyHandle = StructPropertyHandle->GetChildHandle(GET_MEMBER_NAME_CHECKED(FGameplayEffectExecutionScopedModifierInfo, CapturedAttribute));
+
+	TSharedPtr<IPropertyHandle> ParentArrayHandle = StructPropertyHandle->GetParentHandle();
+	const bool bIsExecutionDefAttribute = (ParentArrayHandle.IsValid() && ParentArrayHandle->GetProperty()->GetOuter() == FGameplayEffectExecutionDefinition::StaticStruct());
+	
+	if (bIsExecutionDefAttribute)
+	{
+		TArray<const void*> StructPtrs;
+		StructPropertyHandle->AccessRawData(StructPtrs);
+
+		// Only allow changing the backing definition while single-editing
+		if (StructPtrs.Num() == 1)
+		{
+			TSharedPtr<IPropertyHandle> ExecutionDefinitionHandle = ParentArrayHandle->GetParentHandle();
+			if (ExecutionDefinitionHandle.IsValid())
+			{
+				TArray<const void*> ExecutionDefStructs;
+				ExecutionDefinitionHandle->AccessRawData(ExecutionDefStructs);
+
+				if (ExecutionDefStructs.Num() == 1)
+				{
+					// Extract all of the valid capture definitions off of the capture class
+					const FGameplayEffectExecutionDefinition& ExecutionDef = *reinterpret_cast<const FGameplayEffectExecutionDefinition*>(ExecutionDefStructs[0]);
+					if (ExecutionDef.CalculationClass)
+					{
+						const UGameplayEffectExecutionCalculation* ExecCalcCDO = ExecutionDef.CalculationClass->GetDefaultObject<UGameplayEffectExecutionCalculation>();
+						if (ensure(ExecCalcCDO))
+						{
+							const TArray<FGameplayEffectAttributeCaptureDefinition>& CDOCaptureDefs = ExecCalcCDO->GetAttributeCaptureDefinitions();
+							for (const FGameplayEffectAttributeCaptureDefinition& CurCaptureDef : CDOCaptureDefs)
+							{
+								AvailableCaptureDefs.Add(MakeShareable(new FGameplayEffectAttributeCaptureDefinition(CurCaptureDef)));
+							}
+						}
+					}
+				}
+			}
+		}
+
+		// Construct a custom combo box widget outlining possible capture definition choices
+		if (AvailableCaptureDefs.Num() > 0)
+		{
+			TSharedPtr< SComboBox< TSharedPtr<FGameplayEffectAttributeCaptureDefinition> > > BackingComboBox;
+
+			StructBuilder.AddChildContent(TEXT("CaptureDef"))
+			.NameContent()
+			[
+				SNew(STextBlock)
+				.Text(NSLOCTEXT("ScopedModifierDetails", "CaptureDefLabel", "Backing Capture Definition"))
+				.ToolTipText(NSLOCTEXT("ScopedModifierDetails", "CaptureDefTooltip", "The capture definition to use to populate the scoped modifier. Only options specified by the execution class are presented here."))
+				.Font(StructCustomizationUtils.GetRegularFont())
+			]
+			.ValueContent()
+			.MinDesiredWidth(350.f)
+			[
+				SAssignNew(BackingComboBox, SComboBox< TSharedPtr<FGameplayEffectAttributeCaptureDefinition> >)
+				.OptionsSource(&AvailableCaptureDefs)
+				.OnSelectionChanged(this, &FGameplayEffectExecutionScopedModifierInfoDetails::OnCaptureDefComboBoxSelectionChanged)
+				.OnGenerateWidget(this, &FGameplayEffectExecutionScopedModifierInfoDetails::OnGenerateCaptureDefComboWidget)
+				.Content()
+				[
+					SAssignNew(PrimaryCaptureDefWidget, SCaptureDefWidget)
+				]
+			];
+
+			// Set the initial value on the combo box; done this way to intentionally cause a change delegate
+			if (BackingComboBox.IsValid())
+			{
+				BackingComboBox->SetSelectedItem(GetCurrentCaptureDef());
+			}
+		}
+	}
+
+	uint32 NumChildren;
+	StructPropertyHandle->GetNumChildren(NumChildren);
+
+	// Add all of the properties, though skip the original captured attribute if inside an execution, as it is using the custom
+	// combo box instead
+	for (uint32 ChildIndex = 0; ChildIndex < NumChildren; ++ChildIndex)
+	{
+		const TSharedRef<IPropertyHandle> ChildHandle = StructPropertyHandle->GetChildHandle(ChildIndex).ToSharedRef();
+		const FName ChildPropName = ChildHandle->GetProperty()->GetFName();
+
+		if (!bIsExecutionDefAttribute || (ChildPropName != GET_MEMBER_NAME_CHECKED(FGameplayEffectExecutionScopedModifierInfo, CapturedAttribute)))
+		{
+			StructBuilder.AddChildProperty(ChildHandle);
+		}
+	}
+}
+
+void FGameplayEffectExecutionScopedModifierInfoDetails::OnCaptureDefComboBoxSelectionChanged(TSharedPtr<FGameplayEffectAttributeCaptureDefinition> InSelectedItem, ESelectInfo::Type InSelectInfo)
+{
+	SetCurrentCaptureDef(InSelectedItem);
+
+	// Need to update the base capture widget in the combo box manually due to caching strategy
+	if (PrimaryCaptureDefWidget.IsValid())
+	{
+		PrimaryCaptureDefWidget->SetBackingDefinition(*InSelectedItem.Get());
+	}
+}
+
+TSharedRef<SWidget> FGameplayEffectExecutionScopedModifierInfoDetails::OnGenerateCaptureDefComboWidget(TSharedPtr<FGameplayEffectAttributeCaptureDefinition> InItem)
+{
+	TSharedPtr<SCaptureDefWidget> NewCapDefWidget;
+	SAssignNew(NewCapDefWidget, SCaptureDefWidget);
+	NewCapDefWidget->SetBackingDefinition(*InItem.Get());
+
+	return NewCapDefWidget.ToSharedRef();
+}
+
+TSharedPtr<FGameplayEffectAttributeCaptureDefinition> FGameplayEffectExecutionScopedModifierInfoDetails::GetCurrentCaptureDef() const
+{
+	if (CaptureDefPropertyHandle.IsValid())
+	{
+		TArray<const void*> RawStructPtrs;
+		CaptureDefPropertyHandle->AccessRawData(RawStructPtrs);
+
+		// Only showing the combo box for single-editing
+		const FGameplayEffectAttributeCaptureDefinition& BackingDef = *reinterpret_cast<const FGameplayEffectAttributeCaptureDefinition*>(RawStructPtrs[0]);
+
+		for (TSharedPtr<FGameplayEffectAttributeCaptureDefinition> CurCaptureDef : AvailableCaptureDefs)
+		{
+			if (CurCaptureDef.IsValid() && (*CurCaptureDef.Get() == BackingDef))
+			{
+				return CurCaptureDef;
+			}
+		}
+	}
+
+	return AvailableCaptureDefs[0];
+}
+
+void FGameplayEffectExecutionScopedModifierInfoDetails::SetCurrentCaptureDef(TSharedPtr<FGameplayEffectAttributeCaptureDefinition> InDef)
+{
+	if (CaptureDefPropertyHandle.IsValid() && InDef.IsValid())
+	{
+		TArray<void*> RawStructPtrs;
+		CaptureDefPropertyHandle->AccessRawData(RawStructPtrs);
+
+		FGameplayEffectAttributeCaptureDefinition& BackingDef = *reinterpret_cast<FGameplayEffectAttributeCaptureDefinition*>(RawStructPtrs[0]);
+		const FGameplayEffectAttributeCaptureDefinition& InDefRef = *InDef.Get();
+		if (BackingDef != InDefRef)
+		{
+			CaptureDefPropertyHandle->NotifyPreChange();
+			BackingDef = InDefRef;
+			CaptureDefPropertyHandle->NotifyPostChange();
+		}
+	}
 }
 
-#undef LOCTEXT_NAMESPACE
\ No newline at end of file
+#undef LOCTEXT_NAMESPACE
diff --git a/Engine/Source/Editor/GameplayAbilitiesEditor/Private/GameplayEffectExecutionScopedModifierInfoDetails.h b/Engine/Source/Editor/GameplayAbilitiesEditor/Private/GameplayEffectExecutionScopedModifierInfoDetails.h
--- a/Engine/Source/Editor/GameplayAbilitiesEditor/Private/GameplayEffectExecutionScopedModifierInfoDetails.h
+++ b/Engine/Source/Editor/GameplayAbilitiesEditor/Private/GameplayEffectExecutionScopedModifierInfoDetails.h
@@ -7,6 +7,8 @@
 class IPropertyHandle;
 class FDetailWidgetRow;
 class IDetailChildrenBuilder;
+class SCaptureDefWidget;
+struct FGameplayEffectAttributeCaptureDefinition;
 
 /** Details customization for FGameplayEffectExecutionScopedModifierInfo */
 class FGameplayEffectExecutionScopedModifierInfoDetails : public IPropertyTypeCustomization
@@ -14,7 +16,32 @@ class FGameplayEffectExecutionScopedModifierInfoDetails : public IPropertyTypeCu
 public:
 	static TSharedRef<IPropertyTypeCustomization> MakeInstance();
 
-	/** WIP */
+	/** Overridden to provide the property name */
 	virtual void CustomizeHeader(TSharedRef<IPropertyHandle> StructPropertyHandle, FDetailWidgetRow& HeaderRow, IPropertyTypeCustomizationUtils& StructCustomizationUtils) override;
+	
+	/** Overridden to allow for a custom selection widget for scoped modifiers inside a custom execution */
 	virtual void CustomizeChildren(TSharedRef<IPropertyHandle> StructPropertyHandle, IDetailChildrenBuilder& StructBuilder, IPropertyTypeCustomizationUtils& StructCustomizationUtils) override;
+
+private:
+
+	/** Delegate called when combo box selection is changed */
+	void OnCaptureDefComboBoxSelectionChanged(TSharedPtr<FGameplayEffectAttributeCaptureDefinition> InSelectedItem, ESelectInfo::Type InSelectInfo);
+	
+	/* Called to generate the widgets for custom combo box entries */
+	TSharedRef<SWidget> OnGenerateCaptureDefComboWidget(TSharedPtr<FGameplayEffectAttributeCaptureDefinition> InItem);
+	
+	/** Get the current capture definition as specified by the backing property, if possible; Otherwise falls back to first available definition from execution class */
+	TSharedPtr<FGameplayEffectAttributeCaptureDefinition> GetCurrentCaptureDef() const;
+
+	/** Set the current capture definition */
+	void SetCurrentCaptureDef(TSharedPtr<FGameplayEffectAttributeCaptureDefinition> InDef);
+
+	/** Cached property handle for the capture definition property */
+	TSharedPtr<IPropertyHandle> CaptureDefPropertyHandle;
+
+	/** Primary capture definition widget shown for the custom combo box */
+	TSharedPtr<SCaptureDefWidget> PrimaryCaptureDefWidget;
+
+	/** Backing source for the custom combo box; Populated by all valid definitions from the execution class */
+	TArray< TSharedPtr<FGameplayEffectAttributeCaptureDefinition> > AvailableCaptureDefs;
 };
\ No newline at end of file
diff --git a/Engine/Source/Runtime/GameplayAbilities/Public/GameplayEffectTypes.h b/Engine/Source/Runtime/GameplayAbilities/Public/GameplayEffectTypes.h
--- a/Engine/Source/Runtime/GameplayAbilities/Public/GameplayEffectTypes.h
+++ b/Engine/Source/Runtime/GameplayAbilities/Public/GameplayEffectTypes.h
@@ -195,7 +195,7 @@ struct FGameplayModifierEvaluatedData
 
 /** Struct defining gameplay attribute capture options for gameplay effects */
 USTRUCT()
-struct FGameplayEffectAttributeCaptureDefinition
+struct GAMEPLAYABILITIES_API FGameplayEffectAttributeCaptureDefinition
 {
 	GENERATED_USTRUCT_BODY()
 
