- Add details customization for GameplayEffectExecutionDefinition that hides/prunes calculation modifiers as the user changes calculation class

[REDACTED_HANDLE]: Only works for single asset editing atm, which leaves some awkward UI in gameplay effect multi-editing view

diff --git a/Engine/Source/Editor/GameplayAbilitiesEditor/Private/GameplayAbilitiesEditorModule.cpp b/Engine/Source/Editor/GameplayAbilitiesEditor/Private/GameplayAbilitiesEditorModule.cpp
--- a/Engine/Source/Editor/GameplayAbilitiesEditor/Private/GameplayAbilitiesEditorModule.cpp
+++ b/Engine/Source/Editor/GameplayAbilitiesEditor/Private/GameplayAbilitiesEditorModule.cpp
@@ -7,6 +7,7 @@
 #include "GameplayEffectDetails.h"
 #include "GameplayModifierInfoDetails.h"
 #include "GameplayEffectExecutionScopedModifierInfoDetails.h"
+#include "GameplayEffectExecutionDefinitionDetails.h"
 
 #include "IAssetTypeActions.h"
 #include "AssetToolsModule.h"
@@ -45,6 +46,7 @@ void FGameplayAbilitiesEditorModule::StartupModule()
 	PropertyModule.RegisterCustomPropertyTypeLayout( "ScalableFloat", FOnGetPropertyTypeCustomizationInstance::CreateStatic( &FScalableFloatDetails::MakeInstance ) );
 	PropertyModule.RegisterCustomPropertyTypeLayout( "GameplayModifierInfo", FOnGetPropertyTypeCustomizationInstance::CreateStatic( &FGameplayModifierInfoCustomization::MakeInstance ) );
 	PropertyModule.RegisterCustomPropertyTypeLayout( "GameplayEffectExecutionScopedModifierInfo", FOnGetPropertyTypeCustomizationInstance::CreateStatic( &FGameplayEffectExecutionScopedModifierInfoDetails::MakeInstance ) );
+	PropertyModule.RegisterCustomPropertyTypeLayout( "GameplayEffectExecutionDefinition", FOnGetPropertyTypeCustomizationInstance::CreateStatic( &FGameplayEffectExecutionDefinitionDetails::MakeInstance ) );
 
 	PropertyModule.RegisterCustomClassLayout( "AttributeSet", FOnGetDetailCustomizationInstance::CreateStatic( &FAttributeDetails::MakeInstance ) );
 	PropertyModule.RegisterCustomClassLayout( "GameplayEffect", FOnGetDetailCustomizationInstance::CreateStatic( &FGameplayEffectDetails::MakeInstance ) );
@@ -81,6 +83,7 @@ void FGameplayAbilitiesEditorModule::ShutdownModule()
 		PropertyModule.UnregisterCustomClassLayout("GameplayEffect");
 		PropertyModule.UnregisterCustomClassLayout("AttributeSet");
 
+		PropertyModule.UnregisterCustomPropertyTypeLayout("GameplayEffectExecutionDefinition");
 		PropertyModule.UnregisterCustomPropertyTypeLayout("GameplayEffectExecutionScopedModifierInfo");
 		PropertyModule.UnregisterCustomPropertyTypeLayout("GameplayModifierInfo");
 		PropertyModule.UnregisterCustomPropertyTypeLayout("ScalableFloat");
diff --git a/Engine/Source/Editor/GameplayAbilitiesEditor/Private/GameplayEffectExecutionDefinitionDetails.cpp b/Engine/Source/Editor/GameplayAbilitiesEditor/Private/GameplayEffectExecutionDefinitionDetails.cpp
new file mode 100644
--- /dev/null
+++ b/Engine/Source/Editor/GameplayAbilitiesEditor/Private/GameplayEffectExecutionDefinitionDetails.cpp
@@ -0,0 +1,122 @@
+// Copyright 1998-2013 Epic Games, Inc. All Rights Reserved.
+
+#include "AbilitySystemEditorPrivatePCH.h"
+#include "GameplayEffectExecutionDefinitionDetails.h"
+#include "DetailWidgetRow.h"
+#include "IDetailPropertyRow.h"
+#include "PropertyHandle.h"
+#include "GameplayEffect.h"
+#include "GameplayEffectTypes.h"
+#include "GameplayEffectExecutionCalculation.h"
+
+#define LOCTEXT_NAMESPACE "GameplayEffectExecutionDefinitionDetailsCustomization"
+
+TSharedRef<IPropertyTypeCustomization> FGameplayEffectExecutionDefinitionDetails::MakeInstance()
+{
+	return MakeShareable(new FGameplayEffectExecutionDefinitionDetails());
+}
+
+void FGameplayEffectExecutionDefinitionDetails::CustomizeHeader(TSharedRef<IPropertyHandle> StructPropertyHandle, FDetailWidgetRow& HeaderRow, IPropertyTypeCustomizationUtils& StructCustomizationUtils)
+{
+	HeaderRow
+	.NameContent()
+	[
+		StructPropertyHandle->CreatePropertyNameWidget()
+	];
+}
+
+void FGameplayEffectExecutionDefinitionDetails::CustomizeChildren(TSharedRef<IPropertyHandle> StructPropertyHandle, IDetailChildrenBuilder& StructBuilder, IPropertyTypeCustomizationUtils& StructCustomizationUtils)
+{
+	bShowCalculationModifiers = false;
+
+	// @todo: For now, only allow single editing
+	if (StructPropertyHandle->GetNumOuterObjects() == 1)
+	{
+		CalculationClassPropHandle = StructPropertyHandle->GetChildHandle(GET_MEMBER_NAME_CHECKED(FGameplayEffectExecutionDefinition, CalculationClass));
+
+		TSharedPtr<IPropertyHandle> CalcModPropHandle = StructPropertyHandle->GetChildHandle(GET_MEMBER_NAME_CHECKED(FGameplayEffectExecutionDefinition, CalculationModifiers));
+		CalculationModifiersArrayPropHandle = CalcModPropHandle.IsValid() ? CalcModPropHandle->AsArray() : nullptr;
+
+		if (CalculationClassPropHandle.IsValid())
+		{
+			CalculationClassPropHandle->SetOnPropertyValueChanged(FSimpleDelegate::CreateSP(this, &FGameplayEffectExecutionDefinitionDetails::OnCalculationClassChanged));
+			StructBuilder.AddChildProperty(CalculationClassPropHandle.ToSharedRef());
+			UpdateCalculationModifiers();
+		}
+
+		if (CalculationModifiersArrayPropHandle.IsValid())
+		{
+			IDetailPropertyRow& PropRow = StructBuilder.AddChildProperty(CalcModPropHandle.ToSharedRef());
+			PropRow.Visibility(TAttribute<EVisibility>::Create(TAttribute<EVisibility>::FGetter::CreateSP(this, &FGameplayEffectExecutionDefinitionDetails::GetCalculationModifierVisibility)));
+		}
+	}
+}
+
+void FGameplayEffectExecutionDefinitionDetails::OnCalculationClassChanged()
+{
+	UpdateCalculationModifiers();
+}
+
+void FGameplayEffectExecutionDefinitionDetails::UpdateCalculationModifiers()
+{
+	TArray<FGameplayEffectAttributeCaptureDefinition> ValidCaptureDefinitions;
+	
+	// Try to extract the collection of valid capture definitions from the execution class CDO, if possible
+	if (CalculationClassPropHandle.IsValid())
+	{
+		UObject* ObjValue = nullptr;
+		CalculationClassPropHandle->GetValue(ObjValue);
+
+		UClass* ClassObj = Cast<UClass>(ObjValue);
+		if (ClassObj)
+		{
+			const UGameplayEffectExecutionCalculation* ExecutionCDO = ClassObj->GetDefaultObject<UGameplayEffectExecutionCalculation>();
+			if (ExecutionCDO)
+			{
+				ValidCaptureDefinitions.Append(ExecutionCDO->GetAttributeCaptureDefinitions());
+			}
+		}
+	}
+
+	// Want to hide the calculation modifiers if there are no valid definitions
+	bShowCalculationModifiers = (ValidCaptureDefinitions.Num() > 0);
+
+	// Need to prune out any modifiers that are specified for definitions that aren't specified by the execution class
+	if (CalculationModifiersArrayPropHandle.IsValid())
+	{
+		// If there aren't any valid definitions, just dump the whole array
+		if (ValidCaptureDefinitions.Num() == 0)
+		{
+			CalculationModifiersArrayPropHandle->EmptyArray();
+		}
+		// There are some valid definitions, so verify any existing ones to make sure they are in the valid array
+		else
+		{
+			uint32 NumChildren = 0;
+			CalculationModifiersArrayPropHandle->GetNumElements(NumChildren);
+
+			for (int32 ChildIdx = NumChildren - 1; ChildIdx >= 0; --ChildIdx)
+			{
+				TSharedRef<IPropertyHandle> ChildPropHandle = CalculationModifiersArrayPropHandle->GetElement(ChildIdx);
+				
+				TArray<const void*> RawScopedModInfoStructs;
+				ChildPropHandle->AccessRawData(RawScopedModInfoStructs);
+
+				// @todo: For now, only allow single editing
+				ensure(RawScopedModInfoStructs.Num() == 1);
+				const FGameplayEffectExecutionScopedModifierInfo& CurModInfo = *reinterpret_cast<const FGameplayEffectExecutionScopedModifierInfo*>(RawScopedModInfoStructs[0]);
+				if (!ValidCaptureDefinitions.Contains(CurModInfo.CapturedAttribute))
+				{
+					CalculationModifiersArrayPropHandle->DeleteItem(ChildIdx);
+				}
+			}
+		}
+	}
+}
+
+EVisibility FGameplayEffectExecutionDefinitionDetails::GetCalculationModifierVisibility() const
+{
+	return (bShowCalculationModifiers ? EVisibility::Visible : EVisibility::Collapsed);
+}
+
+#undef LOCTEXT_NAMESPACE
diff --git a/Engine/Source/Editor/GameplayAbilitiesEditor/Private/GameplayEffectExecutionDefinitionDetails.h b/Engine/Source/Editor/GameplayAbilitiesEditor/Private/GameplayEffectExecutionDefinitionDetails.h
new file mode 100644
--- /dev/null
+++ b/Engine/Source/Editor/GameplayAbilitiesEditor/Private/GameplayEffectExecutionDefinitionDetails.h
@@ -0,0 +1,42 @@
+// Copyright 1998-2013 Epic Games, Inc. All Rights Reserved.
+
+#pragma once
+
+#include "Editor/PropertyEditor/Public/PropertyEditorModule.h"
+
+class IPropertyHandle;
+class FDetailWidgetRow;
+class IDetailChildrenBuilder;
+
+/** Details customization for FGameplayEffectExecutionDefinition */
+class FGameplayEffectExecutionDefinitionDetails : public IPropertyTypeCustomization
+{
+public:
+	static TSharedRef<IPropertyTypeCustomization> MakeInstance();
+
+	/** Overridden to provide the property name */
+	virtual void CustomizeHeader(TSharedRef<IPropertyHandle> StructPropertyHandle, FDetailWidgetRow& HeaderRow, IPropertyTypeCustomizationUtils& StructCustomizationUtils) override;
+	
+	/** Overridden to allow for hiding/updating of the calculation modifiers array as the calculation class changes */
+	virtual void CustomizeChildren(TSharedRef<IPropertyHandle> StructPropertyHandle, IDetailChildrenBuilder& StructBuilder, IPropertyTypeCustomizationUtils& StructCustomizationUtils) override;
+
+private:
+
+	/** Called via delegate when the user changes the calculation class */
+	void OnCalculationClassChanged();
+
+	/** Helper function to determine whether to hide/show the modifiers, as well as prune invalid options */
+	void UpdateCalculationModifiers();
+
+	/** Visibility delegate for calculation modifiers */
+	EVisibility GetCalculationModifierVisibility() const;
+
+	/** Property handle for the calculation class */
+	TSharedPtr<IPropertyHandle> CalculationClassPropHandle;
+
+	/** Property handle array for the calculation modifiers */
+	TSharedPtr<IPropertyHandleArray> CalculationModifiersArrayPropHandle;
+
+	/** If true, the calculation modifiers array should be shown */
+	bool bShowCalculationModifiers;
+};
\ No newline at end of file
