- New details customization for gameplay effect mod magnitude that only shows the currently selected magnitude calculation type instead of all of them simultaneously
- Some non-unity fix-ups

diff --git a/Engine/Source/Editor/GameplayAbilitiesEditor/Private/GameplayAbilitiesEditorModule.cpp b/Engine/Source/Editor/GameplayAbilitiesEditor/Private/GameplayAbilitiesEditorModule.cpp
--- a/Engine/Source/Editor/GameplayAbilitiesEditor/Private/GameplayAbilitiesEditorModule.cpp
+++ b/Engine/Source/Editor/GameplayAbilitiesEditor/Private/GameplayAbilitiesEditorModule.cpp
@@ -8,6 +8,7 @@
 #include "GameplayModifierInfoDetails.h"
 #include "GameplayEffectExecutionScopedModifierInfoDetails.h"
 #include "GameplayEffectExecutionDefinitionDetails.h"
+#include "GameplayEffectModifierMagnitudeDetails.h"
 
 #include "IAssetTypeActions.h"
 #include "AssetToolsModule.h"
@@ -47,6 +48,7 @@ void FGameplayAbilitiesEditorModule::StartupModule()
 	PropertyModule.RegisterCustomPropertyTypeLayout( "GameplayModifierInfo", FOnGetPropertyTypeCustomizationInstance::CreateStatic( &FGameplayModifierInfoCustomization::MakeInstance ) );
 	PropertyModule.RegisterCustomPropertyTypeLayout( "GameplayEffectExecutionScopedModifierInfo", FOnGetPropertyTypeCustomizationInstance::CreateStatic( &FGameplayEffectExecutionScopedModifierInfoDetails::MakeInstance ) );
 	PropertyModule.RegisterCustomPropertyTypeLayout( "GameplayEffectExecutionDefinition", FOnGetPropertyTypeCustomizationInstance::CreateStatic( &FGameplayEffectExecutionDefinitionDetails::MakeInstance ) );
+	PropertyModule.RegisterCustomPropertyTypeLayout( "GameplayEffectModifierMagnitude", FOnGetPropertyTypeCustomizationInstance::CreateStatic( &FGameplayEffectModifierMagnitudeDetails::MakeInstance ) );
 
 	PropertyModule.RegisterCustomClassLayout( "AttributeSet", FOnGetDetailCustomizationInstance::CreateStatic( &FAttributeDetails::MakeInstance ) );
 	PropertyModule.RegisterCustomClassLayout( "GameplayEffect", FOnGetDetailCustomizationInstance::CreateStatic( &FGameplayEffectDetails::MakeInstance ) );
@@ -83,6 +85,7 @@ void FGameplayAbilitiesEditorModule::ShutdownModule()
 		PropertyModule.UnregisterCustomClassLayout("GameplayEffect");
 		PropertyModule.UnregisterCustomClassLayout("AttributeSet");
 
+		PropertyModule.UnregisterCustomPropertyTypeLayout("GameplayEffectModifierMagnitude");
 		PropertyModule.UnregisterCustomPropertyTypeLayout("GameplayEffectExecutionDefinition");
 		PropertyModule.UnregisterCustomPropertyTypeLayout("GameplayEffectExecutionScopedModifierInfo");
 		PropertyModule.UnregisterCustomPropertyTypeLayout("GameplayModifierInfo");
diff --git a/Engine/Source/Editor/GameplayAbilitiesEditor/Private/GameplayEffectExecutionDefinitionDetails.cpp b/Engine/Source/Editor/GameplayAbilitiesEditor/Private/GameplayEffectExecutionDefinitionDetails.cpp
--- a/Engine/Source/Editor/GameplayAbilitiesEditor/Private/GameplayEffectExecutionDefinitionDetails.cpp
+++ b/Engine/Source/Editor/GameplayAbilitiesEditor/Private/GameplayEffectExecutionDefinitionDetails.cpp
@@ -8,6 +8,7 @@
 #include "GameplayEffect.h"
 #include "GameplayEffectTypes.h"
 #include "GameplayEffectExecutionCalculation.h"
+#include "DetailLayoutBuilder.h"
 
 #define LOCTEXT_NAMESPACE "GameplayEffectExecutionDefinitionDetailsCustomization"
 
diff --git a/Engine/Source/Editor/GameplayAbilitiesEditor/Private/GameplayEffectExecutionScopedModifierInfoDetails.cpp b/Engine/Source/Editor/GameplayAbilitiesEditor/Private/GameplayEffectExecutionScopedModifierInfoDetails.cpp
--- a/Engine/Source/Editor/GameplayAbilitiesEditor/Private/GameplayEffectExecutionScopedModifierInfoDetails.cpp
+++ b/Engine/Source/Editor/GameplayAbilitiesEditor/Private/GameplayEffectExecutionScopedModifierInfoDetails.cpp
@@ -8,6 +8,7 @@
 #include "GameplayEffect.h"
 #include "GameplayEffectTypes.h"
 #include "GameplayEffectExecutionCalculation.h"
+#include "DetailLayoutBuilder.h"
 
 #define LOCTEXT_NAMESPACE "GameplayEffectExecutionScopedModifierInfoDetailsCustomization"
 
diff --git a/Engine/Source/Editor/GameplayAbilitiesEditor/Private/GameplayEffectModifierMagnitudeDetails.cpp b/Engine/Source/Editor/GameplayAbilitiesEditor/Private/GameplayEffectModifierMagnitudeDetails.cpp
new file mode 100644
--- /dev/null
+++ b/Engine/Source/Editor/GameplayAbilitiesEditor/Private/GameplayEffectModifierMagnitudeDetails.cpp
@@ -0,0 +1,81 @@
+// Copyright 1998-2013 Epic Games, Inc. All Rights Reserved.
+
+#include "AbilitySystemEditorPrivatePCH.h"
+#include "GameplayEffectModifierMagnitudeDetails.h"
+#include "DetailWidgetRow.h"
+#include "IDetailPropertyRow.h"
+#include "PropertyHandle.h"
+#include "GameplayEffect.h"
+#include "DetailLayoutBuilder.h"
+
+#define LOCTEXT_NAMESPACE "GameplayEffectModifierMagnitudeDetailsCustomization"
+
+TSharedRef<IPropertyTypeCustomization> FGameplayEffectModifierMagnitudeDetails::MakeInstance()
+{
+	return MakeShareable(new FGameplayEffectModifierMagnitudeDetails());
+}
+
+void FGameplayEffectModifierMagnitudeDetails::CustomizeHeader(TSharedRef<IPropertyHandle> StructPropertyHandle, FDetailWidgetRow& HeaderRow, IPropertyTypeCustomizationUtils& StructCustomizationUtils)
+{
+	HeaderRow
+	.NameContent()
+	[
+		StructPropertyHandle->CreatePropertyNameWidget()
+	];
+}
+
+void FGameplayEffectModifierMagnitudeDetails::CustomizeChildren(TSharedRef<IPropertyHandle> StructPropertyHandle, IDetailChildrenBuilder& StructBuilder, IPropertyTypeCustomizationUtils& StructCustomizationUtils)
+{
+	// Set up acceleration map
+	PropertyToCalcEnumMap.Empty();
+	PropertyToCalcEnumMap.Add(FindFieldChecked<UProperty>(FGameplayEffectModifierMagnitude::StaticStruct(), GET_MEMBER_NAME_CHECKED(FGameplayEffectModifierMagnitude, ScalableFloatMagnitude)), EGameplayEffectMagnitudeCalculation::ScalableFloat);
+	PropertyToCalcEnumMap.Add(FindFieldChecked<UProperty>(FGameplayEffectModifierMagnitude::StaticStruct(), GET_MEMBER_NAME_CHECKED(FGameplayEffectModifierMagnitude, AttributeBasedMagnitude)), EGameplayEffectMagnitudeCalculation::AttributeBased);
+	PropertyToCalcEnumMap.Add(FindFieldChecked<UProperty>(FGameplayEffectModifierMagnitude::StaticStruct(), GET_MEMBER_NAME_CHECKED(FGameplayEffectModifierMagnitude, CalculationClassMagnitude)), EGameplayEffectMagnitudeCalculation::CustomCalculationClass);
+
+	// Hook into calculation type changes
+	MagnitudeCalculationTypePropertyHandle = StructPropertyHandle->GetChildHandle(GET_MEMBER_NAME_CHECKED(FGameplayEffectModifierMagnitude, MagnitudeCalculationType));
+	if (MagnitudeCalculationTypePropertyHandle.IsValid())
+	{
+		MagnitudeCalculationTypePropertyHandle->SetOnPropertyValueChanged(FSimpleDelegate::CreateSP(this, &FGameplayEffectModifierMagnitudeDetails::OnCalculationTypeChanged));
+		StructBuilder.AddChildProperty(MagnitudeCalculationTypePropertyHandle.ToSharedRef());
+	}
+	OnCalculationTypeChanged();
+
+	// Set up visibility delegates on all of the magnitude calculations
+	TSharedPtr<IPropertyHandle> ScalableFloatMagnitudePropertyHandle = StructPropertyHandle->GetChildHandle(GET_MEMBER_NAME_CHECKED(FGameplayEffectModifierMagnitude, ScalableFloatMagnitude));
+	if (ScalableFloatMagnitudePropertyHandle.IsValid())
+	{
+		IDetailPropertyRow& PropRow = StructBuilder.AddChildProperty(ScalableFloatMagnitudePropertyHandle.ToSharedRef());
+		PropRow.Visibility(TAttribute<EVisibility>::Create(TAttribute<EVisibility>::FGetter::CreateSP(this, &FGameplayEffectModifierMagnitudeDetails::GetMagnitudeCalculationPropertyVisibility, ScalableFloatMagnitudePropertyHandle->GetProperty())));
+	}
+
+	TSharedPtr<IPropertyHandle> AttributeBasedMagnitudePropertyHandle  = StructPropertyHandle->GetChildHandle(GET_MEMBER_NAME_CHECKED(FGameplayEffectModifierMagnitude, AttributeBasedMagnitude));
+	if (AttributeBasedMagnitudePropertyHandle.IsValid())
+	{
+		IDetailPropertyRow& PropRow = StructBuilder.AddChildProperty(AttributeBasedMagnitudePropertyHandle.ToSharedRef());
+		PropRow.Visibility(TAttribute<EVisibility>::Create(TAttribute<EVisibility>::FGetter::CreateSP(this, &FGameplayEffectModifierMagnitudeDetails::GetMagnitudeCalculationPropertyVisibility, AttributeBasedMagnitudePropertyHandle->GetProperty())));
+	}
+
+	TSharedPtr<IPropertyHandle> CalculationClassMagnitudePropertyHandle  = StructPropertyHandle->GetChildHandle(GET_MEMBER_NAME_CHECKED(FGameplayEffectModifierMagnitude, CalculationClassMagnitude));
+	if (CalculationClassMagnitudePropertyHandle.IsValid())
+	{
+		IDetailPropertyRow& PropRow = StructBuilder.AddChildProperty(CalculationClassMagnitudePropertyHandle.ToSharedRef());
+		PropRow.Visibility(TAttribute<EVisibility>::Create(TAttribute<EVisibility>::FGetter::CreateSP(this, &FGameplayEffectModifierMagnitudeDetails::GetMagnitudeCalculationPropertyVisibility, CalculationClassMagnitudePropertyHandle->GetProperty())));
+	}
+}
+
+void FGameplayEffectModifierMagnitudeDetails::OnCalculationTypeChanged()
+{
+	uint8 EnumVal;
+	MagnitudeCalculationTypePropertyHandle->GetValue(EnumVal);
+
+	VisibleCalculationType = static_cast<EGameplayEffectMagnitudeCalculation>(EnumVal);
+}
+
+EVisibility FGameplayEffectModifierMagnitudeDetails::GetMagnitudeCalculationPropertyVisibility(UProperty* InProperty) const
+{
+	const EGameplayEffectMagnitudeCalculation* EnumVal = PropertyToCalcEnumMap.Find(InProperty);
+	return ((EnumVal && (*EnumVal == VisibleCalculationType)) ? EVisibility::Visible : EVisibility::Collapsed);
+}
+
+#undef LOCTEXT_NAMESPACE
diff --git a/Engine/Source/Editor/GameplayAbilitiesEditor/Private/GameplayEffectModifierMagnitudeDetails.h b/Engine/Source/Editor/GameplayAbilitiesEditor/Private/GameplayEffectModifierMagnitudeDetails.h
new file mode 100644
--- /dev/null
+++ b/Engine/Source/Editor/GameplayAbilitiesEditor/Private/GameplayEffectModifierMagnitudeDetails.h
@@ -0,0 +1,39 @@
+// Copyright 1998-2013 Epic Games, Inc. All Rights Reserved.
+
+#pragma once
+
+#include "Editor/PropertyEditor/Public/PropertyEditorModule.h"
+
+class IPropertyHandle;
+class FDetailWidgetRow;
+class IDetailChildrenBuilder;
+
+/** Details customization for FGameplayEffectModifierMagnitude */
+class FGameplayEffectModifierMagnitudeDetails : public IPropertyTypeCustomization
+{
+public:
+	static TSharedRef<IPropertyTypeCustomization> MakeInstance();
+
+	/** Overridden to provide the property name */
+	virtual void CustomizeHeader(TSharedRef<IPropertyHandle> StructPropertyHandle, FDetailWidgetRow& HeaderRow, IPropertyTypeCustomizationUtils& StructCustomizationUtils) override;
+	
+	/** Overridden to only show the magnitude that is currently being used based on calculation type */
+	virtual void CustomizeChildren(TSharedRef<IPropertyHandle> StructPropertyHandle, IDetailChildrenBuilder& StructBuilder, IPropertyTypeCustomizationUtils& StructCustomizationUtils) override;
+
+private:
+
+	/** Called via delegate when the user changes the calculation type */
+	void OnCalculationTypeChanged();
+
+	/** Visibility delegate for the various methods of calculating magnitude */
+	EVisibility GetMagnitudeCalculationPropertyVisibility(UProperty* InProperty) const;
+
+	/** Property handle of the enumeration of the magnitude calculation type */
+	TSharedPtr<IPropertyHandle> MagnitudeCalculationTypePropertyHandle;
+
+	/** Currently visible magnitude calculation type */
+	EGameplayEffectMagnitudeCalculation VisibleCalculationType;
+
+	/** Acceleration map for determining whether to show a magnitude property or not */
+	TMap<UProperty*, EGameplayEffectMagnitudeCalculation> PropertyToCalcEnumMap;
+};
\ No newline at end of file
diff --git a/Engine/Source/Runtime/GameplayAbilities/Public/GameplayEffect.h b/Engine/Source/Runtime/GameplayAbilities/Public/GameplayEffect.h
--- a/Engine/Source/Runtime/GameplayAbilities/Public/GameplayEffect.h
+++ b/Engine/Source/Runtime/GameplayAbilities/Public/GameplayEffect.h
@@ -228,6 +228,7 @@ protected:
 
 	// @hack: @todo: This is temporary to aid in post-load fix-up w/o exposing members publicly
 	friend class UGameplayEffect;
+	friend class FGameplayEffectModifierMagnitudeDetails;
 };
 
 /** 
