- Non-unity build fixes

diff --git a/Engine/Source/Editor/GameplayAbilitiesEditor/Private/GameplayEffectExecutionDefinitionDetails.h b/Engine/Source/Editor/GameplayAbilitiesEditor/Private/GameplayEffectExecutionDefinitionDetails.h
--- a/Engine/Source/Editor/GameplayAbilitiesEditor/Private/GameplayEffectExecutionDefinitionDetails.h
+++ b/Engine/Source/Editor/GameplayAbilitiesEditor/Private/GameplayEffectExecutionDefinitionDetails.h
@@ -5,6 +5,7 @@
 #include "Editor/PropertyEditor/Public/PropertyEditorModule.h"
 
 class IPropertyHandle;
+class IPropertyHandleArray;
 class FDetailWidgetRow;
 class IDetailChildrenBuilder;
 
diff --git a/Engine/Source/Editor/GameplayAbilitiesEditor/Private/GameplayEffectModifierMagnitudeDetails.h b/Engine/Source/Editor/GameplayAbilitiesEditor/Private/GameplayEffectModifierMagnitudeDetails.h
--- a/Engine/Source/Editor/GameplayAbilitiesEditor/Private/GameplayEffectModifierMagnitudeDetails.h
+++ b/Engine/Source/Editor/GameplayAbilitiesEditor/Private/GameplayEffectModifierMagnitudeDetails.h
@@ -3,6 +3,7 @@
 #pragma once
 
 #include "Editor/PropertyEditor/Public/PropertyEditorModule.h"
+#include "GameplayEffect.h"
 
 class IPropertyHandle;
 class FDetailWidgetRow;
