- Defer array operations/pruning by a frame to prevent initialization issues with multiple customizations

diff --git a/Engine/Source/Editor/GameplayAbilitiesEditor/Private/GameplayEffectExecutionDefinitionDetails.cpp b/Engine/Source/Editor/GameplayAbilitiesEditor/Private/GameplayEffectExecutionDefinitionDetails.cpp
--- a/Engine/Source/Editor/GameplayAbilitiesEditor/Private/GameplayEffectExecutionDefinitionDetails.cpp
+++ b/Engine/Source/Editor/GameplayAbilitiesEditor/Private/GameplayEffectExecutionDefinitionDetails.cpp
@@ -9,6 +9,7 @@
 #include "GameplayEffectTypes.h"
 #include "GameplayEffectExecutionCalculation.h"
 #include "DetailLayoutBuilder.h"
+#include "IPropertyUtilities.h"
 
 #define LOCTEXT_NAMESPACE "GameplayEffectExecutionDefinitionDetailsCustomization"
 
@@ -42,7 +43,7 @@ void FGameplayEffectExecutionDefinitionDetails::CustomizeChildren(TSharedRef<IPr
 		{
 			CalculationClassPropHandle->SetOnPropertyValueChanged(FSimpleDelegate::CreateSP(this, &FGameplayEffectExecutionDefinitionDetails::OnCalculationClassChanged));
 			StructBuilder.AddChildProperty(CalculationClassPropHandle.ToSharedRef());
-			UpdateCalculationModifiers();
+			StructCustomizationUtils.GetPropertyUtilities()->EnqueueDeferredAction(FSimpleDelegate::CreateSP(this, &FGameplayEffectExecutionDefinitionDetails::UpdateCalculationModifiers));
 		}
 
 		if (CalculationModifiersArrayPropHandle.IsValid())
