- Add ability to specify on custom GE executions which captured attributes are *not* valid for scoped modifiers

diff --git a/Engine/Source/Editor/GameplayAbilitiesEditor/Private/GameplayEffectExecutionDefinitionDetails.cpp b/Engine/Source/Editor/GameplayAbilitiesEditor/Private/GameplayEffectExecutionDefinitionDetails.cpp
--- a/Engine/Source/Editor/GameplayAbilitiesEditor/Private/GameplayEffectExecutionDefinitionDetails.cpp
+++ b/Engine/Source/Editor/GameplayAbilitiesEditor/Private/GameplayEffectExecutionDefinitionDetails.cpp
@@ -75,7 +75,7 @@ void FGameplayEffectExecutionDefinitionDetails::UpdateCalculationModifiers()
 			const UGameplayEffectExecutionCalculation* ExecutionCDO = ClassObj->GetDefaultObject<UGameplayEffectExecutionCalculation>();
 			if (ExecutionCDO)
 			{
-				ValidCaptureDefinitions.Append(ExecutionCDO->GetAttributeCaptureDefinitions());
+				ExecutionCDO->GetValidScopedModifierAttributeCaptureDefinitions(ValidCaptureDefinitions);
 			}
 		}
 	}
diff --git a/Engine/Source/Editor/GameplayAbilitiesEditor/Private/GameplayEffectExecutionScopedModifierInfoDetails.cpp b/Engine/Source/Editor/GameplayAbilitiesEditor/Private/GameplayEffectExecutionScopedModifierInfoDetails.cpp
--- a/Engine/Source/Editor/GameplayAbilitiesEditor/Private/GameplayEffectExecutionScopedModifierInfoDetails.cpp
+++ b/Engine/Source/Editor/GameplayAbilitiesEditor/Private/GameplayEffectExecutionScopedModifierInfoDetails.cpp
@@ -174,10 +174,12 @@ void FGameplayEffectExecutionScopedModifierInfoDetails::CustomizeChildren(TShare
 						const UGameplayEffectExecutionCalculation* ExecCalcCDO = ExecutionDef.CalculationClass->GetDefaultObject<UGameplayEffectExecutionCalculation>();
 						if (ensure(ExecCalcCDO))
 						{
-							const TArray<FGameplayEffectAttributeCaptureDefinition>& CDOCaptureDefs = ExecCalcCDO->GetAttributeCaptureDefinitions();
-							for (const FGameplayEffectAttributeCaptureDefinition& CurCaptureDef : CDOCaptureDefs)
+							TArray<FGameplayEffectAttributeCaptureDefinition> CaptureDefs;
+							ExecCalcCDO->GetValidScopedModifierAttributeCaptureDefinitions(CaptureDefs);
+
+							for (const FGameplayEffectAttributeCaptureDefinition& CurDef : CaptureDefs)
 							{
-								AvailableCaptureDefs.Add(MakeShareable(new FGameplayEffectAttributeCaptureDefinition(CurCaptureDef)));
+								AvailableCaptureDefs.Add(MakeShareable(new FGameplayEffectAttributeCaptureDefinition(CurDef)));
 							}
 						}
 					}
diff --git a/Engine/Source/Runtime/GameplayAbilities/Private/GameplayEffectExecutionCalculation.cpp b/Engine/Source/Runtime/GameplayAbilities/Private/GameplayEffectExecutionCalculation.cpp
--- a/Engine/Source/Runtime/GameplayAbilities/Private/GameplayEffectExecutionCalculation.cpp
+++ b/Engine/Source/Runtime/GameplayAbilities/Private/GameplayEffectExecutionCalculation.cpp
@@ -154,6 +154,22 @@ UGameplayEffectExecutionCalculation::UGameplayEffectExecutionCalculation(const F
 {
 }
 
+#if WITH_EDITORONLY_DATA
+void UGameplayEffectExecutionCalculation::GetValidScopedModifierAttributeCaptureDefinitions(OUT TArray<FGameplayEffectAttributeCaptureDefinition>& OutScopableModifiers) const
+{
+	OutScopableModifiers.Empty();
+
+	const TArray<FGameplayEffectAttributeCaptureDefinition>& DefaultCaptureDefs = GetAttributeCaptureDefinitions();
+	for (const FGameplayEffectAttributeCaptureDefinition& CurDef : DefaultCaptureDefs)
+	{
+		if (!InvalidScopedModifierAttributes.Contains(CurDef))
+		{
+			OutScopableModifiers.Add(CurDef);
+		}
+	}
+}
+#endif // #if WITH_EDITORONLY_DATA
+
 void UGameplayEffectExecutionCalculation::Execute_Implementation(const FGameplayEffectCustomExecutionParameters& ExecutionParams, OUT TArray<FGameplayModifierEvaluatedData>& OutAdditionalModifiers) const
 {
 }
diff --git a/Engine/Source/Runtime/GameplayAbilities/Public/GameplayEffectExecutionCalculation.h b/Engine/Source/Runtime/GameplayAbilities/Public/GameplayEffectExecutionCalculation.h
--- a/Engine/Source/Runtime/GameplayAbilities/Public/GameplayEffectExecutionCalculation.h
+++ b/Engine/Source/Runtime/GameplayAbilities/Public/GameplayEffectExecutionCalculation.h
@@ -91,7 +91,23 @@ class GAMEPLAYABILITIES_API UGameplayEffectExecutionCalculation : public UGamepl
 {
 	GENERATED_UCLASS_BODY()
 
+#if WITH_EDITORONLY_DATA
+
+protected:
+
+	/** Any attribute in this list will not show up as a valid option for scoped modifiers; Used to allow attribute capture for internal calculation while preventing modification */
+	UPROPERTY(EditDefaultsOnly, Category=Attributes)
+	TArray<FGameplayEffectAttributeCaptureDefinition> InvalidScopedModifierAttributes;
+
 public:
+	/**
+	 * Gets the collection of capture attribute definitions that the calculation class will accept as valid scoped modifiers
+	 * 
+	 * @param OutScopableModifiers	[OUT] Array to populate with definitions valid as scoped modifiers
+	 */
+	virtual void GetValidScopedModifierAttributeCaptureDefinitions(OUT TArray<FGameplayEffectAttributeCaptureDefinition>& OutScopableModifiers) const;
+
+#endif // #if WITH_EDITORONLY_DATA
 
 	/**
 	 * Called whenever the owning gameplay effect is executed. Allowed to do essentially whatever is desired, including generating new
