remove dead code

diff --git a/Engine/Source/Editor/GameplayAbilitiesEditor/Private/AttributeDetails.cpp b/Engine/Source/Editor/GameplayAbilitiesEditor/Private/AttributeDetails.cpp
--- a/Engine/Source/Editor/GameplayAbilitiesEditor/Private/AttributeDetails.cpp
+++ b/Engine/Source/Editor/GameplayAbilitiesEditor/Private/AttributeDetails.cpp
@@ -32,65 +32,7 @@ void FAttributePropertyDetails::CustomizeHeader( TSharedRef<IPropertyHandle> Str
 	FString FilterMetaStr = StructPropertyHandle->GetProperty()->GetMetaData(TEXT("FilterMetaTag"));
 
 	TArray<UProperty*> PropertiesToAdd;
-
-
 	
-/*
-	const bool bIsExtensionAttribute = (StructPropertyHandle->GetProperty()->GetOuter() == FExtensionAttributeModifierInfo::StaticStruct());
-	if (bIsExtensionAttribute)
-	{
-		// Attributes in FExtensionAttributeModifierInfo in FGameplayModifierInfos can only be set to the attributes in the extension class
-		TSharedPtr<IPropertyHandle> ExtensionAttributeModifierInfoHandle = StructPropertyHandle->GetParentHandle();
-		TSharedPtr<IPropertyHandle> AttributeModifierArrayHandle = ExtensionAttributeModifierInfoHandle.IsValid() ? ExtensionAttributeModifierInfoHandle->GetParentHandle() : nullptr;
-		TSharedPtr<IPropertyHandle> ModifierCallbackHandle = AttributeModifierArrayHandle.IsValid() ? AttributeModifierArrayHandle->GetParentHandle() : nullptr;
-		if ( ModifierCallbackHandle.IsValid() )
-		{
-			TArray<const void*> RawPtrs;
-			ModifierCallbackHandle->AccessRawData(RawPtrs);
-			if ( RawPtrs.Num() > 0 )
-			{
-				// Only allow setting of attributes of all selected structs have the same extension class
-				const FGameplayModifierCallback& FirstCallback = *reinterpret_cast<const FGameplayModifierCallback*>(RawPtrs[0]);
-				TSubclassOf<UGameplayEffectExtension> CommonExtensionClass = FirstCallback.ExtensionClass;
-				bool bHasCommonExtensionClass = true;
-
-				for( const void* Ptr : RawPtrs )
-				{
-					const FGameplayModifierCallback& Callback = *reinterpret_cast<const FGameplayModifierCallback*>(Ptr);
-					if ( Callback.ExtensionClass != CommonExtensionClass )
-					{
-						bHasCommonExtensionClass = false;
-						break;
-					}
-				}
-
-				if ( bHasCommonExtensionClass && CommonExtensionClass != nullptr )
-				{
-					UGameplayEffectExtension* ExtensionCDO = CommonExtensionClass->GetDefaultObject<UGameplayEffectExtension>();
-					if ( ensure(ExtensionCDO) )
-					{
-						const bool bIsSourceAttributeModifiers = (AttributeModifierArrayHandle->GetProperty()->GetFName() == GET_MEMBER_NAME_CHECKED(FGameplayModifierCallback, SourceAttributeModifiers));
-						if ( bIsSourceAttributeModifiers )
-						{
-							for (const FGameplayAttribute& RelevantAttribute : ExtensionCDO->RelevantSourceAttributes)
-							{
-								PropertiesToAdd.Add(RelevantAttribute.GetUProperty());
-							}
-						}
-						else
-						{
-							for (const FGameplayAttribute& RelevantAttribute : ExtensionCDO->RelevantTargetAttributes)
-							{
-								PropertiesToAdd.Add(RelevantAttribute.GetUProperty());
-							}
-						}
-					}
-				}
-			}
-		}
-	}
-	else
-*/
 	{
 		// Gather all UAttribute classes
 		for (TObjectIterator<UClass> ClassIt; ClassIt; ++ClassIt)
@@ -346,14 +288,6 @@ void FScalableFloatDetails::CustomizeHeader( TSharedRef<class IPropertyHandle> S
 	uint32 NumChildren = 0;
 	StructPropertyHandle->GetNumChildren(NumChildren);
 
-	/* Debug?
-	for (uint32 i = 0; i < NumChildren; i++)
-	{
-		FString PropName = StructPropertyHandle->GetChildHandle(i)->GetPropertyDisplayName();
-		PropName;
-	}
-	*/
-
 	ValueProperty = StructPropertyHandle->GetChildHandle(GET_MEMBER_NAME_CHECKED(FScalableFloat,Value));
 	CurveTableHandleProperty = StructPropertyHandle->GetChildHandle(GET_MEMBER_NAME_CHECKED(FScalableFloat,Curve));
 
diff --git a/Engine/Source/Editor/GameplayAbilitiesEditor/Private/GameplayAbilitiesEditorModule.cpp b/Engine/Source/Editor/GameplayAbilitiesEditor/Private/GameplayAbilitiesEditorModule.cpp
--- a/Engine/Source/Editor/GameplayAbilitiesEditor/Private/GameplayAbilitiesEditorModule.cpp
+++ b/Engine/Source/Editor/GameplayAbilitiesEditor/Private/GameplayAbilitiesEditorModule.cpp
@@ -51,7 +51,6 @@ void FGameplayAbilitiesEditorModule::StartupModule()
 	FPropertyEditorModule& PropertyModule = FModuleManager::LoadModuleChecked<FPropertyEditorModule>("PropertyEditor");
 	PropertyModule.RegisterCustomPropertyTypeLayout( "GameplayAttribute", FOnGetPropertyTypeCustomizationInstance::CreateStatic( &FAttributePropertyDetails::MakeInstance ) );
 	PropertyModule.RegisterCustomPropertyTypeLayout( "ScalableFloat", FOnGetPropertyTypeCustomizationInstance::CreateStatic( &FScalableFloatDetails::MakeInstance ) );
-	PropertyModule.RegisterCustomPropertyTypeLayout( "GameplayModifierInfo", FOnGetPropertyTypeCustomizationInstance::CreateStatic( &FGameplayModifierInfoCustomization::MakeInstance ) );
 	PropertyModule.RegisterCustomPropertyTypeLayout( "GameplayEffectExecutionScopedModifierInfo", FOnGetPropertyTypeCustomizationInstance::CreateStatic( &FGameplayEffectExecutionScopedModifierInfoDetails::MakeInstance ) );
 	PropertyModule.RegisterCustomPropertyTypeLayout( "GameplayEffectExecutionDefinition", FOnGetPropertyTypeCustomizationInstance::CreateStatic( &FGameplayEffectExecutionDefinitionDetails::MakeInstance ) );
 	PropertyModule.RegisterCustomPropertyTypeLayout( "GameplayEffectModifierMagnitude", FOnGetPropertyTypeCustomizationInstance::CreateStatic( &FGameplayEffectModifierMagnitudeDetails::MakeInstance ) );
diff --git a/Engine/Source/Editor/GameplayAbilitiesEditor/Private/GameplayModifierInfoDetails.cpp b/Engine/Source/Editor/GameplayAbilitiesEditor/Private/GameplayModifierInfoDetails.cpp
--- a/Engine/Source/Editor/GameplayAbilitiesEditor/Private/GameplayModifierInfoDetails.cpp
+++ b/Engine/Source/Editor/GameplayAbilitiesEditor/Private/GameplayModifierInfoDetails.cpp
@@ -53,36 +53,6 @@ void FGameplayModifierInfoCustomization::CustomizeChildren(TSharedRef<IPropertyH
 
 void FGameplayModifierInfoCustomization::OnModifierOpChanged(TSharedRef<IPropertyHandle> StructPropertyHandle)
 {
-	// When changing to a callback op, automatically add an initial entry in the callbacks list.
-	// When changing back to a non-callback op, clear the entry if it is the only one and has no data (just to be clean)
-	TArray<void*> StructPtrs;
-	StructPropertyHandle->AccessRawData(StructPtrs);
-	if (StructPtrs.Num() == 1)
-	{
-		FGameplayModifierInfo& Modifier = *reinterpret_cast<FGameplayModifierInfo*>(StructPtrs[0]);
-
-		if ( Modifier.ModifierOp == EGameplayModOp::Callback )
-		{
-			// We switched to Callback. See if we should add an initial callback.
-			if ( Modifier.Callbacks.Num() == 0 )
-			{
-				new(Modifier.Callbacks) FGameplayModifierCallback();
-			}
-		}
-		else
-		{
-			// We switched away from Callback. See if we only have a single item with no data. If so, remove it.
-			if ( Modifier.Callbacks.Num() == 1 )
-			{
-				const FGameplayModifierCallback& Callback = Modifier.Callbacks[0];
-				if ( Callback.ExtensionClass == nullptr && Callback.SourceAttributeModifiers.Num() == 0 && Callback.TargetAttributeModifiers.Num() == 0 )
-				{
-					Modifier.Callbacks.Empty();
-				}
-			}
-		}
-	}
-
 	UpdateHiddenProperties(StructPropertyHandle);
 }
 
diff --git a/Engine/Source/Runtime/GameplayAbilities/Public/GameplayEffect.h b/Engine/Source/Runtime/GameplayAbilities/Public/GameplayEffect.h
--- a/Engine/Source/Runtime/GameplayAbilities/Public/GameplayEffect.h
+++ b/Engine/Source/Runtime/GameplayAbilities/Public/GameplayEffect.h
@@ -19,63 +19,6 @@ class UAbilitySystemComponent;
 class UGameplayModMagnitudeCalculation;
 class UGameplayEffectExecutionCalculation;
 
-// Repurposing(?) this to be the inplace/custom/scoped modifier used by FGameplayEffectExecutionDefinition
-//	I would expect this to be the thing that is built off to be a "scalable float/attribute reference/custom magnitude" thing.
-//	This would potentially need to be snapshot/not
-// 
-USTRUCT()
-struct GAMEPLAYABILITIES_API FExtensionAttributeModifierInfo
-{
-	GENERATED_USTRUCT_BODY()
-
-	FExtensionAttributeModifierInfo()
-		: ModifierOp( EGameplayModOp::Additive )
-	{
-
-	}
-
-	/** The Attribute we modify or the GE we modify modifies. */
-	UPROPERTY(EditDefaultsOnly, Category=GameplayModifier)
-	FGameplayAttribute Attribute;
-
-	/** The numeric operation of this modifier: Override, Add, Multiply, etc  */
-	UPROPERTY(EditDefaultsOnly, Category=GameplayModifier)
-	TEnumAsByte<EGameplayModOp::Type> ModifierOp;
-
-	/** How much this modifies what it is applied to */
-	UPROPERTY(EditDefaultsOnly, Category=GameplayModifier)
-	FScalableFloat Magnitude;
-
-	/** Source of the gameplay attribute */
-	UPROPERTY(EditDefaultsOnly, Category = Capture)
-	EGameplayEffectAttributeCaptureSource Source;
-
-	UPROPERTY(EditDefaultsOnly, Category = GameplayModifier)
-	FGameplayTagRequirements	SourceTags;
-
-	UPROPERTY(EditDefaultsOnly, Category = GameplayModifier)
-	FGameplayTagRequirements	TargetTags;
-};
-
-
-// Should this just go away now?
-USTRUCT()
-struct GAMEPLAYABILITIES_API FGameplayModifierCallback
-{
-	GENERATED_USTRUCT_BODY()
-
-	UPROPERTY(EditDefaultsOnly, Category=GameplayModifier)
-	TSubclassOf<class UGameplayEffectExtension> ExtensionClass;
-
-	/** Modifications to attributes on the source instigator to be used in the extension class */
-	UPROPERTY(EditDefaultsOnly, Category=GameplayModifier)
-	TArray<FExtensionAttributeModifierInfo> SourceAttributeModifiers;
-
-	/** Modifications to attributes on the target to be used in the extension class */
-	UPROPERTY(EditDefaultsOnly, Category=GameplayModifier)
-	TArray<FExtensionAttributeModifierInfo> TargetAttributeModifiers;
-};
-
 USTRUCT()
 struct FGameplayEffectStackingCallbacks
 {
@@ -376,9 +319,6 @@ struct GAMEPLAYABILITIES_API FGameplayModifierInfo
 	UPROPERTY(EditDefaultsOnly, Category=GameplayModifier)
 	TEnumAsByte<EGameplayModOp::Type> ModifierOp;
 
-	UPROPERTY(EditDefaultsOnly, Category=GameplayModifier, meta=(DisplayName="Custom Ops"))
-	TArray<FGameplayModifierCallback> Callbacks;
-
 	// @todo: Remove this after content resave
 	/** Now "deprecated," though being handled in a custom manner to avoid engine version bump. */
 	UPROPERTY()
@@ -1025,7 +965,7 @@ struct FActiveGameplayEffect : public FFastArraySerializerItem
 	UPROPERTY(NotReplicated)
 	float StartWorldTime;
 
-	// Not sure if this should replicate or not. If replicated, we may have trouble where IsHibited doesn't appear to change when we do tag checks (because it was previously inhibited, but replication made it inhibited).
+	// Not sure if this should replicate or not. If replicated, we may have trouble where IsInhibited doesn't appear to change when we do tag checks (because it was previously inhibited, but replication made it inhibited).
 	UPROPERTY()
 	bool IsInhibited;
 
diff --git a/Engine/Source/Runtime/GameplayAbilities/Public/GameplayEffectTypes.h b/Engine/Source/Runtime/GameplayAbilities/Public/GameplayEffectTypes.h
--- a/Engine/Source/Runtime/GameplayAbilities/Public/GameplayEffectTypes.h
+++ b/Engine/Source/Runtime/GameplayAbilities/Public/GameplayEffectTypes.h
@@ -43,7 +43,6 @@ namespace EGameplayModOp
 
 		// Other
 		Override 			UMETA(DisplayName="Override"),	// This should always be the first non numeric ModOp
-		Callback			UMETA(DisplayName="Custom"),
 
 		// This must always be at the end
 		Max					UMETA(DisplayName="Invalid")
