remove include

diff --git a/Engine/Source/Editor/GameplayAbilitiesEditor/Private/GameplayAbilitiesEditorModule.cpp b/Engine/Source/Editor/GameplayAbilitiesEditor/Private/GameplayAbilitiesEditorModule.cpp
--- a/Engine/Source/Editor/GameplayAbilitiesEditor/Private/GameplayAbilitiesEditorModule.cpp
+++ b/Engine/Source/Editor/GameplayAbilitiesEditor/Private/GameplayAbilitiesEditorModule.cpp
@@ -5,7 +5,6 @@
 #include "AttributeDetails.h"
 #include "AttributeSet.h"
 #include "GameplayEffectDetails.h"
-#include "GameplayModifierInfoDetails.h"
 #include "GameplayEffectExecutionScopedModifierInfoDetails.h"
 #include "GameplayEffectExecutionDefinitionDetails.h"
 #include "GameplayEffectModifierMagnitudeDetails.h"
