Editor - Preventing a crash on shutdown, checking if the UObject system is initialized when the Gameplay Abilitiies editor module shuts down.

diff --git a/Engine/Source/Editor/GameplayAbilitiesEditor/Private/GameplayAbilitiesEditorModule.cpp b/Engine/Source/Editor/GameplayAbilitiesEditor/Private/GameplayAbilitiesEditorModule.cpp
--- a/Engine/Source/Editor/GameplayAbilitiesEditor/Private/GameplayAbilitiesEditorModule.cpp
+++ b/Engine/Source/Editor/GameplayAbilitiesEditor/Private/GameplayAbilitiesEditorModule.cpp
@@ -138,7 +138,7 @@ void FGameplayAbilitiesEditorModule::ShutdownModule()
 		GameplayAbilitiesGraphPanelNodeFactory.Reset();
 	}
 
-	if ( IGameplayTagsModule::IsAvailable() )
+	if ( UObjectInitialized() && IGameplayTagsModule::IsAvailable() )
 	{
 		UGameplayTagsManager& GameplayTagsManager = IGameplayTagsModule::Get().GetGameplayTagsManager();
 		GameplayTagsManager.OnGameplayTagTreeChanged().Remove(GameplayTagTreeChangedDelegateHandle);
