fix nullptr crash in FScalableFloatDetails

diff --git a/Engine/Source/Editor/GameplayAbilitiesEditor/Private/AttributeDetails.cpp b/Engine/Source/Editor/GameplayAbilitiesEditor/Private/AttributeDetails.cpp
--- a/Engine/Source/Editor/GameplayAbilitiesEditor/Private/AttributeDetails.cpp
+++ b/Engine/Source/Editor/GameplayAbilitiesEditor/Private/AttributeDetails.cpp
@@ -573,7 +573,7 @@ FText FScalableFloatDetails::GetRowValuePreviewText() const
 {
 	TArray<const void*> RawPtrs;
 	CurveTableHandleProperty->AccessRawData(RawPtrs);
-	if ( RawPtrs.Num() == 1 )
+	if ( RawPtrs.Num() == 1 && RawPtrs[0] )
 	{
 		const FCurveTableRowHandle& Curve = *reinterpret_cast<const FCurveTableRowHandle*>(RawPtrs[0]);
 		if ( Curve.CurveTable && Curve.RowName != NAME_None )
