[UE-13143] - Disabling references to blueprint defined properties because the pointers go stale and crashes occur. This has never worked, except possibly during the editor session in which the reference is assigned.

diff --git a/Engine/Source/Editor/GameplayAbilitiesEditor/Private/AttributeDetails.cpp b/Engine/Source/Editor/GameplayAbilitiesEditor/Private/AttributeDetails.cpp
--- a/Engine/Source/Editor/GameplayAbilitiesEditor/Private/AttributeDetails.cpp
+++ b/Engine/Source/Editor/GameplayAbilitiesEditor/Private/AttributeDetails.cpp
@@ -38,7 +38,7 @@ void FAttributePropertyDetails::CustomizeHeader( TSharedRef<IPropertyHandle> Str
 		for (TObjectIterator<UClass> ClassIt; ClassIt; ++ClassIt)
 		{
 			UClass *Class = *ClassIt;
-			if (Class->IsChildOf(UAttributeSet::StaticClass()) && !FKismetEditorUtilities::IsClassABlueprintSkeleton(Class))
+			if (Class->IsChildOf(UAttributeSet::StaticClass()) && !Class->ClassGeneratedBy)
 			{
 				// Allow entire classes to be filtered globally
 				if (Class->HasMetaData(TEXT("HideInDetailsView")))
@@ -66,7 +66,7 @@ void FAttributePropertyDetails::CustomizeHeader( TSharedRef<IPropertyHandle> Str
 			}
 
 			// UAbilitySystemComponent can add 'system' attributes
-			if (Class->IsChildOf(UAbilitySystemComponent::StaticClass()) && !FKismetEditorUtilities::IsClassABlueprintSkeleton(Class))
+			if (Class->IsChildOf(UAbilitySystemComponent::StaticClass()) && !Class->ClassGeneratedBy)
 			{
 				for (TFieldIterator<UProperty> PropertyIt(Class, EFieldIteratorFlags::ExcludeSuper); PropertyIt; ++PropertyIt)
 				{
