Making the comment more general.

diff --git a/Engine/Source/Editor/GameplayAbilitiesEditor/Private/SGameplayAttributeGraphPin.cpp b/Engine/Source/Editor/GameplayAbilitiesEditor/Private/SGameplayAttributeGraphPin.cpp
--- a/Engine/Source/Editor/GameplayAbilitiesEditor/Private/SGameplayAttributeGraphPin.cpp
+++ b/Engine/Source/Editor/GameplayAbilitiesEditor/Private/SGameplayAttributeGraphPin.cpp
@@ -18,7 +18,7 @@ void SGameplayAttributeGraphPin::Construct( const FArguments& InArgs, UEdGraphPi
 TSharedRef<SWidget>	SGameplayAttributeGraphPin::GetDefaultValueWidget()
 {
 	// Parse out current default value
-	// It will be in the form (Attribute=/Script/OrionGame.OrionHealthSet:Health)
+	// It will be in the form (Attribute=/Script/<PackageName>.<ObjectName>:<PropertyName>)
 	
 	FString DefaultString = GraphPinObj->GetDefaultAsString();
 	FGameplayAttribute DefaultAttribute;
