UE-14868 Adding new, Async Proxy output pin to UK2Node_BaseAsyncTask messed up AbilityTasks that spawn actors

diff --git a/Engine/Source/Editor/GameplayAbilitiesEditor/Private/K2Node_LatentAbilityCall.cpp b/Engine/Source/Editor/GameplayAbilitiesEditor/Private/K2Node_LatentAbilityCall.cpp
--- a/Engine/Source/Editor/GameplayAbilitiesEditor/Private/K2Node_LatentAbilityCall.cpp
+++ b/Engine/Source/Editor/GameplayAbilitiesEditor/Private/K2Node_LatentAbilityCall.cpp
@@ -572,13 +572,20 @@ void UK2Node_LatentAbilityCall::ExpandNode(class FKismetCompilerContext& Compile
 			}
 		}
 	}
+
+	// Expose Async Task Proxy object
+	UEdGraphPin* const ProxyObjectPin = CallCreateProxyObjectNode->GetReturnValuePin();
+	check(ProxyObjectPin);
+	auto OutputAsyncTaskProxy = FindPinChecked(FBaseAsyncTaskHelper::GetAsyncTaskProxyName());
+	bIsErrorFree &= CompilerContext.MovePinLinksToIntermediate(*OutputAsyncTaskProxy, *ProxyObjectPin).CanSafeConnect();
+
 	// ------------------------------------------------------------------------------------------
 	// GATHER OUTPUT PARAMETERS AND PAIR THEM WITH LOCAL VARIABLES
 	// ------------------------------------------------------------------------------------------
 	TArray<FBaseAsyncTaskHelper::FOutputPinAndLocalVariable> VariableOutputs;
 	for (auto CurrentPin : Pins)
 	{
-		if (FBaseAsyncTaskHelper::ValidDataPin(CurrentPin, EGPD_Output, Schema))
+		if ((OutputAsyncTaskProxy != CurrentPin) && FBaseAsyncTaskHelper::ValidDataPin(CurrentPin, EGPD_Output, Schema))
 		{
 			const FEdGraphPinType& PinType = CurrentPin->PinType;
 			UK2Node_TemporaryVariable* TempVarOutput = CompilerContext.SpawnInternalVariable(
@@ -592,7 +599,6 @@ void UK2Node_LatentAbilityCall::ExpandNode(class FKismetCompilerContext& Compile
 	// FOR EACH DELEGATE DEFINE EVENT, CONNECT IT TO DELEGATE AND IMPLEMENT A CHAIN OF ASSIGMENTS
 	// ------------------------------------------------------------------------------------------
 	UEdGraphPin* LastThenPin = CallCreateProxyObjectNode->FindPinChecked(Schema->PN_Then);
-	UEdGraphPin* const ProxyObjectPin = CallCreateProxyObjectNode->GetReturnValuePin();
 	for (TFieldIterator<UMulticastDelegateProperty> PropertyIt(ProxyClass, EFieldIteratorFlags::ExcludeSuper); PropertyIt && bIsErrorFree; ++PropertyIt)
 	{
 		bIsErrorFree &= FBaseAsyncTaskHelper::HandleDelegateImplementation(*PropertyIt, VariableOutputs, ProxyObjectPin, LastThenPin, this, SourceGraph, CompilerContext);
