null check

--------
Integrated using branch UE4-To-UE4-Orion (reversed) of change#2557572 by [REDACTED_PERSON] on 2015/05/19 15:55:56.

diff --git a/Engine/Source/Editor/GameplayAbilitiesEditor/Private/K2Node_LatentAbilityCall.cpp b/Engine/Source/Editor/GameplayAbilitiesEditor/Private/K2Node_LatentAbilityCall.cpp
--- a/Engine/Source/Editor/GameplayAbilitiesEditor/Private/K2Node_LatentAbilityCall.cpp
+++ b/Engine/Source/Editor/GameplayAbilitiesEditor/Private/K2Node_LatentAbilityCall.cpp
@@ -192,6 +192,8 @@ void UK2Node_LatentAbilityCall::CreatePinsForClass(UClass* InClass)
 	// (For example, hide Instigator in situations where instigator is not relevant to your task)
 	
 	TArray<FString> IgnorePropertyList;
+
+	if (ProxyFactoryClass)
 	{
 		UFunction* ProxyFunction = ProxyFactoryClass->FindFunctionByName(ProxyFactoryFunctionName);
 
