GameplayCue Editor tool first pass

--------
Integrated using branch UE4-To-UE4-Orion (reversed) of change#2583077 by [REDACTED_PERSON] on 2015/06/10 14:39:58.

diff --git a/Engine/Source/Editor/GameplayAbilitiesEditor/Private/GameplayAbilitiesEditorModule.cpp b/Engine/Source/Editor/GameplayAbilitiesEditor/Private/GameplayAbilitiesEditorModule.cpp
--- a/Engine/Source/Editor/GameplayAbilitiesEditor/Private/GameplayAbilitiesEditorModule.cpp
+++ b/Engine/Source/Editor/GameplayAbilitiesEditor/Private/GameplayAbilitiesEditorModule.cpp
@@ -20,6 +20,10 @@
 #include "Editor/BlueprintGraph/Public/BlueprintActionDatabase.h"
 #include "K2Node_GameplayCueEvent.h"
 
+#include "SGameplayCueEditor.h"
+#include "Editor/WorkspaceMenuStructure/Public/WorkspaceMenuStructureModule.h"
+#include "SDockTab.h"
+
 class FGameplayAbilitiesEditorModule : public IGameplayAbilitiesEditorModule
 {
 	// Begin IModuleInterface
@@ -27,6 +31,30 @@ class FGameplayAbilitiesEditorModule : public IGameplayAbilitiesEditorModule
 	virtual void ShutdownModule() override;
 	// End IModuleInterface
 
+	TSharedRef<SDockTab> SpawnGameplayCueEditorTab(const FSpawnTabArgs& Args);
+	TSharedPtr<SWidget> SummonGameplayCueEditorUI();
+
+	FGetGameplayCueNotifyClasses& GetGameplayCueNotifyClassesDelegate() override
+	{
+		return GetGameplayCueNotifyClasses;
+	}
+
+	FGetGameplayCuePath& GetGameplayCueNotifyPathDelegate()
+	{
+		return GetGameplayCueNotifyPath;
+	}
+
+	FGetGameplayCueInterfaceClasses& GetGameplayCueInterfaceClassesDelegate()
+	{
+		return GetGameplayCueInterfaceClasses;
+
+	}
+
+	FGetGameplayCueEditorStrings& GetGameplayCueEditorStringsDelegate()
+	{
+		return GetGameplayCueEditorStrings;
+	}
+
 protected:
 	void RegisterAssetTypeAction(class IAssetTools& AssetTools, TSharedRef<IAssetTypeActions> Action);
 	static void GameplayTagTreeChanged();
@@ -43,6 +71,14 @@ private:
 
 	/** Handle to the registered GameplayTagTreeChanged delegate */
 	FDelegateHandle GameplayTagTreeChangedDelegateHandle;
+
+	FGetGameplayCueNotifyClasses GetGameplayCueNotifyClasses;
+
+	FGetGameplayCuePath GetGameplayCueNotifyPath;
+
+	FGetGameplayCueInterfaceClasses GetGameplayCueInterfaceClasses;
+	
+	FGetGameplayCueEditorStrings GetGameplayCueEditorStrings;
 };
 
 IMPLEMENT_MODULE(FGameplayAbilitiesEditorModule, GameplayAbilitiesEditor)
@@ -77,6 +113,14 @@ void FGameplayAbilitiesEditorModule::StartupModule()
 	// Listen for changes to the gameplay tag tree so we can refresh blueprint actions for the GameplayCueEvent node
 	UGameplayTagsManager& GameplayTagsManager = IGameplayTagsModule::Get().GetGameplayTagsManager();
 	GameplayTagTreeChangedDelegateHandle = GameplayTagsManager.OnGameplayTagTreeChanged().AddStatic(&FGameplayAbilitiesEditorModule::GameplayTagTreeChanged);
+
+	// GameplayCue editor
+	FGlobalTabmanager::Get()->RegisterNomadTabSpawner( FName(TEXT("GameplayCueApp")), FOnSpawnTab::CreateRaw(this, &FGameplayAbilitiesEditorModule::SpawnGameplayCueEditorTab))
+		.SetDisplayName(NSLOCTEXT("GameplayAbilitiesEditorModule", "GameplayCueTabTitle", "GameplayCue Editor"))
+		.SetTooltipText(NSLOCTEXT("GameplayAbilitiesEditorModule", "GameplayCueTooltipText", "Open GameplayCue Editor tab."))
+		.SetGroup(WorkspaceMenu::GetMenuStructure().GetToolsCategory())
+		//.SetGroup(WorkspaceMenu::GetMenuStructure().GetDeveloperToolsDebugCategory());
+		.SetIcon(FSlateIcon(FEditorStyle::GetStyleSetName(), "Profiler.EventGraph.ExpandHotPath16"));
 }
 
 void FGameplayAbilitiesEditorModule::RegisterAssetTypeAction(IAssetTools& AssetTools, TSharedRef<IAssetTypeActions> Action)
@@ -88,6 +132,10 @@ void FGameplayAbilitiesEditorModule::RegisterAssetTypeAction(IAssetTools& AssetT
 void FGameplayAbilitiesEditorModule::GameplayTagTreeChanged()
 {
 	// The tag tree changed so we should refresh which actions are provided by the gameplay cue event
+#if STATS
+	FString PerfMessage = FString::Printf(TEXT("FGameplayAbilitiesEditorModule::GameplayTagTreeChanged"));
+	SCOPE_LOG_TIME_IN_SECONDS(*PerfMessage, nullptr)
+#endif
 	FBlueprintActionDatabase::Get().RefreshClassActions(UK2Node_GameplayCueEvent::StaticClass());
 }
 
@@ -144,3 +192,24 @@ void FGameplayAbilitiesEditorModule::ShutdownModule()
 		GameplayTagsManager.OnGameplayTagTreeChanged().Remove(GameplayTagTreeChangedDelegateHandle);
 	}
 }
+
+
+TSharedRef<SDockTab> FGameplayAbilitiesEditorModule::SpawnGameplayCueEditorTab(const FSpawnTabArgs& Args)
+{
+	return SNew(SDockTab)
+		.TabRole(ETabRole::NomadTab)
+		[
+			SummonGameplayCueEditorUI().ToSharedRef()
+		];
+}
+
+
+TSharedPtr<SWidget> FGameplayAbilitiesEditorModule::SummonGameplayCueEditorUI()
+{
+	TSharedPtr<SWidget> ReturnWidget;
+	if( IsInGameThread() )
+	{
+		ReturnWidget = SNew(SGameplayCueEditor);			
+	}
+	return ReturnWidget;
+}
\ No newline at end of file
diff --git a/Engine/Source/Editor/GameplayAbilitiesEditor/Private/SGameplayCueEditor.cpp b/Engine/Source/Editor/GameplayAbilitiesEditor/Private/SGameplayCueEditor.cpp
new file mode 100644
--- /dev/null
+++ b/Engine/Source/Editor/GameplayAbilitiesEditor/Private/SGameplayCueEditor.cpp
@@ -0,0 +1,893 @@
+// Copyright 1998-2015 Epic Games, Inc. All Rights Reserved.
+
+#include "AbilitySystemEditorPrivatePCH.h"
+#include "SGameplayCueEditor.h"
+#include "GameplayCueSet.h"
+#include "GameplayCueManager.h"
+#include "Editor/ContentBrowser/Public/ContentBrowserModule.h"
+#include "GameplayCueNotify_Static.h"
+#include "GameplayCueNotify_Actor.h"
+#include "SExpandableArea.h"
+#include "ClassIconFinder.h"
+#include "EditorClassUtils.h"
+#include "GameplayTagsSettings.h"
+#include "SNotificationList.h"
+#include "NotificationManager.h"
+
+#define LOCTEXT_NAMESPACE "SGameplayCueEditor"
+
+
+BEGIN_SLATE_FUNCTION_BUILD_OPTIMIZATION
+
+/**
+ *	Constructs the main widget for the GameplayCueEditor
+ */
+void SGameplayCueEditor::Construct(const FArguments& InArgs)
+{
+	bShowAll = true;
+	bool CanAddFromINI = UGameplayTagsManager::ShouldImportTagsFromINI(); // We only support adding new tags to the ini files.
+	ChildSlot
+	[
+		SNew(SVerticalBox)
+
+		+SVerticalBox::Slot()
+		.AutoHeight( )
+		[
+			SNew( SHorizontalBox )
+			+SHorizontalBox::Slot()
+			.Padding( 2.0f, 2.0f )
+			.AutoWidth()
+			[
+					
+				SNew(SButton)
+				.Text(LOCTEXT("SearchBPEvents", "Search BP Events"))
+				.OnClicked( this, &SGameplayCueEditor::BuildEventMap )
+			]
+
+			+SHorizontalBox::Slot()
+			.Padding( 2.0f, 2.0f )
+			.AutoWidth()
+			[
+				// Check box that controls LIVE MODE
+				SNew(SCheckBox)
+				.IsChecked(this, &SGameplayCueEditor::HandleShowAllCheckBoxIsChecked)
+				.OnCheckStateChanged(this, &SGameplayCueEditor::HandleShowAllCheckedStateChanged)
+				[
+					SNew(STextBlock)
+					.Text(LOCTEXT("HieUnhandled", "Hide Unhandled GameplayCues"))
+				]
+			]
+		]
+
+		+SVerticalBox::Slot()
+		.AutoHeight( )
+		[
+			SNew( SHorizontalBox )
+			+SHorizontalBox::Slot()
+			.Padding( 2.0f, 2.0f )
+			.AutoWidth()
+			[
+				SAssignNew(NewGameplayCueTextBox, SEditableTextBox)
+				.MinDesiredWidth(210.0f)
+				.HintText(LOCTEXT("GameplayCueXY", "GameplayCue.X.Y"))
+				.OnTextCommitted( this, &SGameplayCueEditor::OnNewGameplayCueTagCommited )
+			]
+
+			+SHorizontalBox::Slot()
+			.Padding( 2.0f, 2.0f )
+			.AutoWidth()
+			[
+				SNew(SButton)
+				.Text(LOCTEXT("AddNew", "Add New"))
+				.OnClicked( this, &SGameplayCueEditor::OnNewGameplayCueButtonPressed )
+				.Visibility( CanAddFromINI ? EVisibility::Visible : EVisibility::Collapsed )
+			]
+		]
+
+		+SVerticalBox::Slot()
+		.AutoHeight( )
+		[
+			SNew( SHorizontalBox )
+			+SHorizontalBox::Slot()
+			.Padding( 2.0f, 2.0f )
+			.AutoWidth()
+			[
+				SNew(SSearchBox)
+				.MinDesiredWidth(210.0f)				
+				.OnTextCommitted( this, &SGameplayCueEditor::OnSearchTagCommited )
+			]
+		]
+		
+		+SVerticalBox::Slot()
+		.FillHeight( 1.0f )
+		[
+			SNew( SBorder )
+			.BorderImage( FEditorStyle::GetBrush( "ToolBar.Background" ) )
+			[
+				SAssignNew( GameplayCueListView, SGameplayCueListView )
+					.ItemHeight( 24 )
+					.ListItemsSource( &GameplayCueListItems )
+					.OnGenerateRow( this, &SGameplayCueEditor::OnGenerateWidgetForGameplayCueListView )
+					.HeaderRow
+					(
+						SNew(SHeaderRow)
+						+SHeaderRow::Column("GameplayCueTags")
+						.DefaultLabel(NSLOCTEXT("GameplayCueEditor", "GameplayCueTag", "Tag"))
+						.FillWidth(0.25)
+
+						+SHeaderRow::Column("GameplayCueHandlers")
+						.DefaultLabel(NSLOCTEXT("GameplayCueEditor", "GameplayCueHandlers", "Handlers"))
+						//.FillWidth(1000)
+					)
+			]
+		]
+	];
+
+	UpdateGameplayCueListItems();
+}
+
+END_SLATE_FUNCTION_BUILD_OPTIMIZATION
+
+SGameplayCueEditor::~SGameplayCueEditor()
+{
+	
+}
+
+void SGameplayCueEditor::OnSearchTagCommited(const FText& InText, ETextCommit::Type InCommitType)
+{
+	if (InCommitType == ETextCommit::OnEnter || InCommitType == ETextCommit::OnCleared)
+	{
+		SearchString = InText.ToString();
+		UpdateGameplayCueListItems();
+	}
+}
+
+void SGameplayCueEditor::OnNewGameplayCueTagCommited(const FText& InText, ETextCommit::Type InCommitType)
+{
+	// Only support adding tags via ini file
+	if (UGameplayTagsManager::ShouldImportTagsFromINI() == false)
+	{
+		return;
+	}
+
+	if (InCommitType == ETextCommit::OnEnter)
+	{
+		CreateNewGameplayCueTag();
+	}
+}
+
+FReply SGameplayCueEditor::OnNewGameplayCueButtonPressed()
+{
+	CreateNewGameplayCueTag();
+	return FReply::Handled();
+}
+
+/**
+ *	Checks out config file, adds new tag, repopulates widget cue list
+ */
+void SGameplayCueEditor::CreateNewGameplayCueTag()
+{
+	FScopedSlowTask SlowTask(0.f, LOCTEXT("AddingNewGameplaycue", "Adding new GameplayCue Tag"));
+	SlowTask.MakeDialog();
+				
+
+	FString str = NewGameplayCueTextBox->GetText().ToString();
+
+	AutoSelectTag = str;
+	
+	UGameplayTagsSettings* Settings = GetMutableDefault<UGameplayTagsSettings>();
+	
+	if (Settings)
+	{
+		
+		FString RelativeConfigFilePath = Settings->GetDefaultConfigFilename();
+		FString ConfigPath = FPaths::ConvertRelativePathToFull(RelativeConfigFilePath);
+
+		if (ISourceControlModule::Get().IsEnabled())
+		{
+			FText ErrorMessage;
+
+			if (!SourceControlHelpers::CheckoutOrMarkForAdd(ConfigPath, FText::FromString(ConfigPath), NULL, ErrorMessage))
+			{
+				FNotificationInfo Info(ErrorMessage);
+				Info.ExpireDuration = 3.0f;
+				FSlateNotificationManager::Get().AddNotification(Info);
+			}
+		}
+		else
+		{
+			if (!FPlatformFileManager::Get().GetPlatformFile().SetReadOnly(*ConfigPath, false))
+			{
+				FText NotificationErrorText = FText::Format(LOCTEXT("FailedToMakeWritable", "Could not make {0} writable."), FText::FromString(ConfigPath));
+
+				FNotificationInfo Info(NotificationErrorText);
+				Info.ExpireDuration = 3.0f;
+
+				FSlateNotificationManager::Get().AddNotification(Info);
+			}
+		}
+		
+		Settings->GameplayTags.Add(str);			
+		IGameplayTagsModule::Get().GetGameplayTagsManager().DestroyGameplayTagTree();			
+
+		{
+#if STATS
+			FString PerfMessage = FString::Printf(TEXT("ConstructGameplayTagTree GameplayTag tables"));
+			SCOPE_LOG_TIME_IN_SECONDS(*PerfMessage, nullptr)
+#endif
+			IGameplayTagsModule::Get().GetGameplayTagsManager().ConstructGameplayTagTree();
+		}
+
+		Settings->UpdateDefaultConfigFile();
+
+		UpdateGameplayCueListItems();
+	}
+
+	NewGameplayCueTextBox->SetText(FText::GetEmpty());
+}
+
+void SGameplayCueEditor::HandleShowAllCheckedStateChanged( ECheckBoxState NewValue )
+{
+	bShowAll = (NewValue == ECheckBoxState::Unchecked);
+	UpdateGameplayCueListItems();
+}	
+
+/** Callback for getting the checked state of the focus check box. */
+ECheckBoxState SGameplayCueEditor::HandleShowAllCheckBoxIsChecked() const
+{
+	return bShowAll ? ECheckBoxState::Unchecked : ECheckBoxState::Checked;
+}
+
+/** Builds widget for rows in the GameplayCue Editor tab */
+TSharedRef<ITableRow> SGameplayCueEditor::OnGenerateWidgetForGameplayCueListView(TSharedPtr< FGameplayCueEditorItem > InItem, const TSharedRef<STableViewBase>& OwnerTable)
+{
+	class SGameplayCueItemWidget : public SMultiColumnTableRow< TSharedPtr< FGameplayCueEditorItem > >
+	{
+		public:
+		SLATE_BEGIN_ARGS( SGameplayCueItemWidget ){}
+		SLATE_END_ARGS()
+
+		void Construct( const FArguments& InArgs, const TSharedRef<STableViewBase>& InOwnerTable, TSharedPtr<FGameplayCueEditorItem> InListItem )
+		{
+			Item = InListItem;
+	
+			SMultiColumnTableRow< TSharedPtr< FGameplayCueEditorItem > >::Construct( FSuperRowType::FArguments(), InOwnerTable );			
+		}
+
+		TSharedRef<SWidget> GenerateWidgetForColumn( const FName& ColumnName )
+		{
+			if ( ColumnName == "GameplayCueTags" )
+			{
+				return
+					SNew( STextBlock )
+					.Text( FText::FromString(Item->GameplayCueTagName.ToString()) );
+			}
+			else if ( ColumnName == "GameplayCueHandlers" )
+			{
+				return 
+				SNew( SVerticalBox )				
+				// Notifies
+				+ SVerticalBox::Slot()
+					.AutoHeight()
+					.Padding( 2.0f, 2.0f )
+				[
+					SAssignNew( Item->GameplayCueHandlerListView, SGameplayCueHandlerListView )
+					.ItemHeight( 24 )
+					.ListItemsSource( &Item->HandlerItems )
+					.OnGenerateRow( this, &SGameplayCueItemWidget::OnGenerateWidgetForGameplayCueHandlerListView )
+				]
+	
+				// Add New button
+				+ SVerticalBox::Slot()
+					.AutoHeight()
+					.Padding( 2.0f, 0.0f )
+				[
+					SNew( SHorizontalBox )
+					+ SHorizontalBox::Slot()
+					.AutoWidth()
+					.Padding( 0.0f, 0.0f )
+					[					
+						SNew( SButton )
+							.Text( LOCTEXT("AddNew", "Add New") )
+							.OnClicked( Item.ToSharedRef(), &FGameplayCueEditorItem::OnAddNewClicked )
+					]
+				];
+			}
+			else
+			{
+				return SNew(STextBlock) .Text(LOCTEXT("UnknownColumn", "Unknown Column"));
+			}
+
+		}
+
+		TSharedRef<ITableRow> OnGenerateWidgetForGameplayCueHandlerListView(TSharedPtr< FGameplayCueEditorHandlerItem > InItem, const TSharedRef<STableViewBase>& OwnerTable)
+		{			
+			FLinearColor Color(0.2f, 0.4f, 0.6f, 1.f);
+			FString ObjName;
+			
+			if ( InItem->GameplayCueNotifyObj.AssetLongPathname.IsEmpty() == false)
+			{
+				ObjName = InItem->GameplayCueNotifyObj.AssetLongPathname;
+			
+				int32 idx;
+				if (ObjName.FindLastChar(TEXT('.'), idx))
+				{
+					ObjName = ObjName.RightChop(idx+1);
+					if (ObjName.FindLastChar(TEXT('_'), idx))
+					{
+						ObjName = ObjName.Left(idx);
+					}				
+				}
+			}
+			else if (InItem->FunctionPtr.IsValid())
+			{
+				Color = FLinearColor(1.f, 1.f, 1.0f, 1.f);
+				UFunction* Func = InItem->FunctionPtr.Get();
+				UClass* OuterClass = Cast<UClass>(Func->GetOuter());
+				if (OuterClass)
+				{
+					ObjName = OuterClass->GetName();
+					ObjName.RemoveFromEnd(TEXT("_c"));
+				}
+			}
+
+			return	SNew( STableRow< TSharedPtr<FGameplayCueEditorHandlerItem> >, OwnerTable )
+			.Content()
+			[	
+				SNew( STextBlock )
+				.Text( FText::FromString(ObjName ) )
+				.ColorAndOpacity(Color)
+				.OnDoubleClicked( this, &SGameplayCueItemWidget::OnHandlerDoubleClicked, InItem )
+			];
+		}
+
+		FReply OnHandlerDoubleClicked(TSharedPtr< FGameplayCueEditorHandlerItem > InItem)
+		{
+			if (InItem->GameplayCueNotifyObj.IsValid())
+			{
+				TArray<FString>	AssetList;
+				AssetList.Add(InItem->GameplayCueNotifyObj.ToString());
+				FAssetEditorManager::Get().OpenEditorsForAssets(AssetList);
+			}
+			else if(InItem->FunctionPtr.IsValid())
+			{
+				TArray<FString>	AssetList;
+				AssetList.Add(InItem->FunctionPtr->GetOuter()->GetPathName());
+				FAssetEditorManager::Get().OpenEditorsForAssets(AssetList);
+			}
+			
+			return FReply::Handled();
+		}
+
+		TSharedPtr< FGameplayCueEditorItem > Item;
+	};
+
+	return SNew( SGameplayCueItemWidget, OwnerTable, InItem );
+}
+
+
+
+void SGameplayCueEditor::FGameplayCueEditorItem::OnNewClassPicked(FString SelectedPath)
+{
+	
+}
+
+/** Slow task: use asset registry to find blueprints, load an inspect them to see what GC events they implement. */
+FReply SGameplayCueEditor::BuildEventMap()
+{
+	FScopedSlowTask SlowTask(100.f, LOCTEXT("BuildEventMap", "Searching Blueprints for GameplayCue events"));
+	SlowTask.MakeDialog();
+	SlowTask.EnterProgressFrame(10.f);
+
+	EventMap.Empty();
+
+	IGameplayTagsModule& GameplayTagModule = IGameplayTagsModule::Get();
+
+	auto del = IGameplayAbilitiesEditorModule::Get().GetGameplayCueInterfaceClassesDelegate();
+	if (del.IsBound())
+	{
+		TArray<UClass*> InterfaceClasses;
+		del.ExecuteIfBound(InterfaceClasses);
+		float WorkPerClass = InterfaceClasses.Num() > 0 ? 90.f / static_cast<float>(InterfaceClasses.Num()) : 0.f;
+
+		for (UClass* InterfaceClass : InterfaceClasses)
+		{
+			SlowTask.EnterProgressFrame(WorkPerClass);
+
+			TArray<FAssetData> GameplayCueInterfaceActors;
+			{
+		#if STATS
+				FString PerfMessage = FString::Printf(TEXT("Searched asset registry %s "), *InterfaceClass->GetName());
+				SCOPE_LOG_TIME_IN_SECONDS(*PerfMessage, nullptr)
+		#endif
+
+				UObjectLibrary* ObjLibrary = UObjectLibrary::CreateLibrary(InterfaceClass, true, true);
+				ObjLibrary->LoadBlueprintAssetDataFromPath(TEXT("/Game/"));
+				ObjLibrary->GetAssetDataList(GameplayCueInterfaceActors);
+			}
+			
+			{
+		#if STATS
+				FString PerfMessage = FString::Printf(TEXT("Fully Loaded GameplayCueNotify actors %s "), *InterfaceClass->GetName());
+				SCOPE_LOG_TIME_IN_SECONDS(*PerfMessage, nullptr)
+		#endif				
+
+				for (const FAssetData& AssetData : GameplayCueInterfaceActors)
+				{
+					UBlueprint* BP = Cast<UBlueprint>(AssetData.GetAsset());
+					if (BP)
+					{
+						for (TFieldIterator<UFunction> FuncIt(BP->GeneratedClass, EFieldIteratorFlags::ExcludeSuper); FuncIt; ++FuncIt)
+						{
+							FString FuncName = FuncIt->GetName();
+							if (FuncName.Contains("GameplayCue"))
+							{
+								FuncName.ReplaceInline(TEXT("_"), TEXT("."));
+								FGameplayTag FoundTag = GameplayTagModule.RequestGameplayTag(FName(*FuncName), false);
+								if (FoundTag.IsValid())
+								{
+									EventMap.AddUnique(FoundTag, *FuncIt);
+								}
+							}
+
+						}
+					}
+				}
+			}
+		}
+
+		UpdateGameplayCueListItems();
+	}
+
+	return FReply::Handled();
+}
+
+/** Builds content of the list in the GameplayCue Editor */
+void SGameplayCueEditor::UpdateGameplayCueListItems()
+{
+#if STATS
+	FString PerfMessage = FString::Printf(TEXT("SGameplayCueEditor::UpdateGameplayCueListItems"));
+	SCOPE_LOG_TIME_IN_SECONDS(*PerfMessage, nullptr)
+#endif
+
+	GameplayCueListItems.Reset();
+	IGameplayTagsModule& GameplayTagModule = IGameplayTagsModule::Get();
+	UGameplayCueManager* CueManager = UAbilitySystemGlobals::Get().GetGameplayCueManager();
+	TSharedPtr<FGameplayCueEditorItem> SelectItem;
+
+	FGameplayTagContainer AllGameplayCueTags = IGameplayTagsModule::Get().GetGameplayTagsManager().RequestGameplayTagChildren(UGameplayCueSet::BaseGameplayCueTag());
+	for (FGameplayTag ThisGameplayCueTag : AllGameplayCueTags)
+	{
+		if (SearchString.IsEmpty() == false && ThisGameplayCueTag.ToString().Contains(SearchString) == false)
+		{
+			continue;
+		}
+
+		bool Handled = false;
+
+		TSharedRef< FGameplayCueEditorItem > NewItem( new FGameplayCueEditorItem() );
+		NewItem->GameplayCueTagName = ThisGameplayCueTag;
+
+		if (AutoSelectTag.IsEmpty() == false && AutoSelectTag == ThisGameplayCueTag.ToString())
+		{
+			SelectItem = NewItem;
+		}
+
+		// Add notifies from the global set
+		if (CueManager && CueManager->GlobalCueSet)
+		{
+			int32* idxPtr = CueManager->GlobalCueSet->GameplayCueDataMap.Find(ThisGameplayCueTag);
+			if (idxPtr)
+			{
+				int32 idx = *idxPtr;
+				if (CueManager->GlobalCueSet->GameplayCueData.IsValidIndex(idx))
+				{
+					FGameplayCueNotifyData& Data = CueManager->GlobalCueSet->GameplayCueData[*idxPtr];
+
+					TSharedRef< FGameplayCueEditorHandlerItem > NewHandlerItem( new FGameplayCueEditorHandlerItem() );
+					NewHandlerItem->GameplayCueNotifyObj = Data.GameplayCueNotifyObj;
+					NewItem->HandlerItems.Add(NewHandlerItem);
+					Handled = true;
+				}
+			}
+		}
+
+		// Add events implemented by IGameplayCueInterface blueprints
+		TArray<UFunction*> Funcs;
+		EventMap.MultiFind(ThisGameplayCueTag, Funcs);
+
+		for (UFunction* Func : Funcs)
+		{
+			TSharedRef< FGameplayCueEditorHandlerItem > NewHandlerItem( new FGameplayCueEditorHandlerItem() );
+			NewHandlerItem->FunctionPtr = Func;
+			NewItem->HandlerItems.Add(NewHandlerItem);
+			Handled = true;
+		}
+
+		if (bShowAll || Handled)
+		{
+			GameplayCueListItems.Add(NewItem);
+		}
+
+		// Todo: Add notifies from all GameplayCueSets?
+	}
+
+
+	if (GameplayCueListView.IsValid())
+	{
+		GameplayCueListView->RequestListRefresh();
+	}
+
+	if (SelectItem.IsValid())
+	{
+		GameplayCueListView->SetItemSelection(SelectItem, true);
+		GameplayCueListView->RequestScrollIntoView(SelectItem);
+		
+	}
+	AutoSelectTag.Empty();
+}
+
+/** Widget for picking a new GameplayCue Notify class (similiar to actor class picker)  */
+class SGameplayCuePickerDialog : public SCompoundWidget
+{
+public:
+	SLATE_BEGIN_ARGS( SGameplayCuePickerDialog )
+	{}
+
+		SLATE_ARGUMENT( TSharedPtr<SWindow>, ParentWindow )
+		SLATE_ARGUMENT( TArray<UClass*>, DefaultClasses )
+		SLATE_ARGUMENT( FString, GameplayCueTag )
+
+	SLATE_END_ARGS()
+
+	/** Constructs this widget with InArgs */
+	void Construct(const FArguments& InArgs);
+
+	static bool PickGameplayCue(const FText& TitleText, const TArray<UClass*>& DefaultClasses, UClass*& OutChosenClass, FString InGameplayCueTag);
+
+private:
+	/** Handler for when a class is picked in the class picker */
+	void OnClassPicked(UClass* InChosenClass);
+
+	/** Creates the default class widgets */
+	TSharedRef<ITableRow> GenerateListRow(UClass* InItem, const TSharedRef<STableViewBase>& OwnerTable);
+
+	/** Handler for when "Ok" we selected in the class viewer */
+	FReply OnClassPickerConfirmed();
+
+	FReply OnDefaultClassPicked(UClass* InChosenClass);
+
+	/** Handler for the custom button to hide/unhide the default class viewer */
+	void OnDefaultAreaExpansionChanged(bool bExpanded);
+
+	/** Handler for the custom button to hide/unhide the class viewer */
+	void OnCustomAreaExpansionChanged(bool bExpanded);	
+
+private:
+	/** A pointer to the window that is asking the user to select a parent class */
+	TWeakPtr<SWindow> WeakParentWindow;
+
+	/** The class that was last clicked on */
+	UClass* ChosenClass;
+
+	/** A flag indicating that Ok was selected */
+	bool bPressedOk;
+
+	/**  An array of default classes used in the dialog **/
+	TArray<UClass*> DefaultClasses;
+
+	FString GameplayCueTag;
+};
+
+void SGameplayCuePickerDialog::Construct(const FArguments& InArgs)
+{
+	WeakParentWindow = InArgs._ParentWindow;
+	DefaultClasses = InArgs._DefaultClasses;
+	GameplayCueTag = InArgs._GameplayCueTag;
+
+	FLinearColor AssetColor = FLinearColor::White;
+
+	bPressedOk = false;
+	ChosenClass = NULL;
+
+	FString PathStr = SGameplayCueEditor::GetPathNameForGameplayCueTag(GameplayCueTag);
+	FGameplayCueEditorStrings Strings;
+	auto del = IGameplayAbilitiesEditorModule::Get().GetGameplayCueEditorStringsDelegate();
+	if (del.IsBound())
+	{
+		Strings = del.Execute();
+	}
+
+	ChildSlot
+	[
+		SNew(SBorder)
+		.Visibility(EVisibility::Visible)
+		.BorderImage(FEditorStyle::GetBrush("Menu.Background"))
+		[
+			SNew(SBox)
+			.Visibility(EVisibility::Visible)
+			.Padding(2.f)
+			.WidthOverride(520.0f)
+			[
+				SNew(SVerticalBox)
+				+SVerticalBox::Slot()
+				.Padding(2.f, 2.f)
+				.AutoHeight()
+				[
+					SNew(SBorder)
+					.Visibility(EVisibility::Visible)
+					.BorderImage( FEditorStyle::GetBrush("AssetThumbnail.AssetBackground") )
+					.BorderBackgroundColor(AssetColor.CopyWithNewOpacity(0.3f))
+					[
+						SNew(SExpandableArea)
+						.AreaTitle(NSLOCTEXT("SGameplayCuePickerDialog", "CommonClassesAreaTitle", "GameplayCue Notifies"))
+						.BodyContent()
+						[
+							SNew(SVerticalBox)
+							+SVerticalBox::Slot()
+							.Padding(2.f, 2.f)
+							.AutoHeight()
+							[
+								SNew(STextBlock)
+								
+								.Text(FText::FromString(Strings.GameplayCueNotifyDescription1))
+								.AutoWrapText(true)
+							]
+
+							+SVerticalBox::Slot()
+							.AutoHeight( )
+							[
+								SNew(SListView < UClass*  >)
+								.ItemHeight(48)
+								.SelectionMode(ESelectionMode::None)
+								.ListItemsSource(&DefaultClasses)
+								.OnGenerateRow(this, &SGameplayCuePickerDialog::GenerateListRow)
+							]
+
+							+SVerticalBox::Slot()
+							.Padding(2.f, 2.f)
+							.AutoHeight()
+							[
+								SNew(STextBlock)
+								.Text(FText::FromString(FString::Printf(TEXT("This will create a new GameplayCue Notify here:"))))
+								.AutoWrapText(true)
+							]
+							+SVerticalBox::Slot()
+							.Padding(2.f, 2.f)
+							.AutoHeight()
+							[
+								SNew(STextBlock)
+								.Text(FText::FromString(PathStr))
+								.HighlightText(FText::FromString(PathStr))
+								.AutoWrapText(true)
+							]
+							+SVerticalBox::Slot()
+							.Padding(2.f, 2.f)
+							.AutoHeight()
+							[
+								SNew(STextBlock)
+								.Text(FText::FromString(Strings.GameplayCueNotifyDescription2))
+								.AutoWrapText(true)
+							]
+						]
+					]
+				]
+
+				+SVerticalBox::Slot()
+				.Padding(2.f, 2.f)
+				.AutoHeight()
+				[
+					SNew(SBorder)
+					.Visibility(EVisibility::Visible)
+					.BorderImage( FEditorStyle::GetBrush("AssetThumbnail.AssetBackground") )
+					.BorderBackgroundColor(AssetColor.CopyWithNewOpacity(0.3f))
+					[
+						SNew(SExpandableArea)
+						.AreaTitle(NSLOCTEXT("SGameplayCuePickerDialogEvents", "CommonClassesAreaTitleEvents", "Custom BP Events"))
+						.BodyContent()
+						[
+							SNew(SVerticalBox)
+							+SVerticalBox::Slot()
+							.Padding(2.f, 2.f)
+							.AutoHeight()
+							[
+								SNew(STextBlock)
+								.Text(FText::FromString(Strings.GameplayCueEventDescription1))
+								.AutoWrapText(true)
+
+							]
+
+							+SVerticalBox::Slot()
+							.Padding(2.f, 2.f)
+							.AutoHeight()
+							[
+								SNew(STextBlock)
+								.Text(FText::FromString(Strings.GameplayCueEventDescription2))
+								.AutoWrapText(true)
+							]							
+						]
+					]
+				]
+			]
+		]
+	];
+}
+
+/** Spawns window for picking new GameplayCue handler/notify */
+bool SGameplayCuePickerDialog::PickGameplayCue(const FText& TitleText, const TArray<UClass*>& DefaultClasses, UClass*& OutChosenClass, FString InGameplayCueName)
+{
+	// Create the window to pick the class
+	TSharedRef<SWindow> PickerWindow = SNew(SWindow)
+		.Title(TitleText)
+		.SizingRule( ESizingRule::Autosized )
+		.ClientSize( FVector2D( 0.f, 600.f ))
+		.SupportsMaximize(false)
+		.SupportsMinimize(false);
+
+	TSharedRef<SGameplayCuePickerDialog> ClassPickerDialog = SNew(SGameplayCuePickerDialog)
+		.ParentWindow(PickerWindow)
+		.DefaultClasses(DefaultClasses)
+		.GameplayCueTag(InGameplayCueName);
+
+	PickerWindow->SetContent(ClassPickerDialog);
+
+	GEditor->EditorAddModalWindow(PickerWindow);
+
+	if (ClassPickerDialog->bPressedOk)
+	{
+		OutChosenClass = ClassPickerDialog->ChosenClass;
+		return true;
+	}
+	else
+	{
+		// Ok was not selected, NULL the class
+		OutChosenClass = NULL;
+		return false;
+	}
+}
+
+void SGameplayCuePickerDialog::OnClassPicked(UClass* InChosenClass)
+{
+	ChosenClass = InChosenClass;
+}
+
+/** Generates rows in the list of GameplayCueNotify classes to pick from */
+TSharedRef<ITableRow> SGameplayCuePickerDialog::GenerateListRow(UClass* ItemClass, const TSharedRef<STableViewBase>& OwnerTable)
+{	
+	const FSlateBrush* ItemBrush = FClassIconFinder::FindIconForClass(ItemClass);
+
+	return 
+	SNew(STableRow< UClass* >, OwnerTable)
+	[
+		SNew(SVerticalBox)
+		+SVerticalBox::Slot()
+		.MaxHeight(60.0f)
+		.Padding(10.0f, 6.0f, 0.0f, 4.0f)
+		[
+			SNew(SHorizontalBox)
+			+SHorizontalBox::Slot()
+			.FillWidth(0.45f)
+			[
+				SNew(SButton)
+				.OnClicked(this, &SGameplayCuePickerDialog::OnDefaultClassPicked, ItemClass)
+				.Content()
+				[
+					SNew(SHorizontalBox)
+					+SHorizontalBox::Slot()
+					.HAlign(HAlign_Center)
+					.VAlign(VAlign_Center)
+					.FillWidth(0.12f)
+					[
+						SNew(SImage)
+						.Image(ItemBrush)
+					]
+					+SHorizontalBox::Slot()
+					.VAlign(VAlign_Center)
+					.Padding(4.0f, 0.0f)
+					.FillWidth(0.8f)
+					[
+						SNew(STextBlock)
+						.Text(ItemClass->GetDisplayNameText())
+					]
+
+				]
+			]
+			+SHorizontalBox::Slot()
+			.Padding(10.0f, 0.0f)
+			[
+				SNew(STextBlock)
+				.Text(ItemClass->GetToolTipText(true))
+				.AutoWrapText(true)
+			]
+			+SHorizontalBox::Slot()
+			.AutoWidth()
+			[
+				FEditorClassUtils::GetDocumentationLinkWidget(ItemClass)
+			]
+		]
+	];
+}
+
+FReply SGameplayCuePickerDialog::OnDefaultClassPicked(UClass* InChosenClass)
+{
+	ChosenClass = InChosenClass;
+	bPressedOk = true;
+	if (WeakParentWindow.IsValid())
+	{
+		WeakParentWindow.Pin()->RequestDestroyWindow();
+	}
+	return FReply::Handled();
+}
+
+FReply SGameplayCuePickerDialog::OnClassPickerConfirmed()
+{
+	if (ChosenClass == NULL)
+	{
+		FMessageDialog::Open(EAppMsgType::Ok, NSLOCTEXT("EditorFactories", "MustChooseClassWarning", "You must choose a class."));
+	}
+	else
+	{
+		bPressedOk = true;
+
+		if (WeakParentWindow.IsValid())
+		{
+			WeakParentWindow.Pin()->RequestDestroyWindow();
+		}
+	}
+	return FReply::Handled();
+}
+
+FString SGameplayCueEditor::GetPathNameForGameplayCueTag(FString GameplayCueTagName)
+{
+	FString NewDefaultPathName = FString::Printf(TEXT("/GC_%s"), *GameplayCueTagName);
+	auto PathDel = IGameplayAbilitiesEditorModule::Get().GetGameplayCueNotifyPathDelegate();
+	if (PathDel.IsBound())
+	{
+		NewDefaultPathName = PathDel.Execute(GameplayCueTagName);
+	}
+	NewDefaultPathName.ReplaceInline(TEXT("."), TEXT("_"));
+	return NewDefaultPathName;
+}
+
+/** Create new GameplayCueNotify: brings up dialogue to pick class, then creates it via the content browser. */
+FReply SGameplayCueEditor::FGameplayCueEditorItem::OnAddNewClicked()
+{
+	FAssetToolsModule& AssetToolsModule = FModuleManager::GetModuleChecked<FAssetToolsModule>("AssetTools");
+	FContentBrowserModule& ContentBrowserModule = FModuleManager::LoadModuleChecked<FContentBrowserModule>("ContentBrowser");
+
+	TArray<UClass*>	NotifyClasses;
+	auto del = IGameplayAbilitiesEditorModule::Get().GetGameplayCueNotifyClassesDelegate();
+	del.ExecuteIfBound(NotifyClasses);
+	if (NotifyClasses.Num() == 0)
+	{
+		NotifyClasses.Add(UGameplayCueNotify_Static::StaticClass());
+		NotifyClasses.Add(AGameplayCueNotify_Actor::StaticClass());
+	}
+
+	// --------------------------------------------------
+
+	// Null the parent class to ensure one is selected	
+
+	const FText TitleText = LOCTEXT("CreateBlueprintOptions", "New GameplayCue Handler");
+	UClass* ChosenClass = nullptr;
+	const bool bPressedOk = SGameplayCuePickerDialog::PickGameplayCue(TitleText, NotifyClasses, ChosenClass, GameplayCueTagName.ToString());
+	if (bPressedOk)
+	{
+		FString NewDefaultPathName = SGameplayCueEditor::GetPathNameForGameplayCueTag(GameplayCueTagName.ToString());
+
+		// Make sure the name is unique
+		FString AssetName;
+		FString PackageName;
+		AssetToolsModule.Get().CreateUniqueAssetName(NewDefaultPathName , TEXT(""), /*out*/ PackageName, /*out*/ AssetName);
+		const FString PackagePath = FPackageName::GetLongPackagePath(PackageName);
+
+		// Create the GameplayCue Notify
+		UBlueprintFactory* BlueprintFactory = NewObject<UBlueprintFactory>();
+		BlueprintFactory->ParentClass = ChosenClass;
+		ContentBrowserModule.Get().CreateNewAsset(AssetName, PackagePath, UBlueprint::StaticClass(), BlueprintFactory);
+	}
+	
+	return FReply::Handled();
+}
+
+#undef LOCTEXT_NAMESPACE
diff --git a/Engine/Source/Editor/GameplayAbilitiesEditor/Private/SGameplayCueEditor.h b/Engine/Source/Editor/GameplayAbilitiesEditor/Private/SGameplayCueEditor.h
new file mode 100644
--- /dev/null
+++ b/Engine/Source/Editor/GameplayAbilitiesEditor/Private/SGameplayCueEditor.h
@@ -0,0 +1,80 @@
+// Copyright 1998-2015 Epic Games, Inc. All Rights Reserved.
+
+#pragma once
+
+#include "SCompoundWidget.h"
+
+/** Main CollisionAnalyzer UI widget */
+class SGameplayCueEditor : public SCompoundWidget
+{
+public:
+
+	SLATE_BEGIN_ARGS(SGameplayCueEditor) {}
+
+	SLATE_END_ARGS()
+
+	void Construct(const FArguments& InArgs);
+
+	virtual ~SGameplayCueEditor();
+
+	void OnNewGameplayCueTagCommited(const FText& InText, ETextCommit::Type InCommitType);
+	void OnSearchTagCommited(const FText& InText, ETextCommit::Type InCommitType);
+
+	FReply OnNewGameplayCueButtonPressed();
+	
+	static FString GetPathNameForGameplayCueTag(FString Tag);
+
+private:
+
+	TSharedPtr<SEditableTextBox>	NewGameplayCueTextBox;
+
+	void CreateNewGameplayCueTag();
+
+	void HandleShowAllCheckedStateChanged( ECheckBoxState NewValue );
+
+	/** Callback for getting the checked state of the focus check box. */
+	ECheckBoxState HandleShowAllCheckBoxIsChecked() const;
+
+	bool bShowAll;
+
+	struct FGameplayCueEditorHandlerItem
+	{
+		FStringAssetReference GameplayCueNotifyObj;
+
+		TWeakObjectPtr<UFunction> FunctionPtr;
+	};
+
+	typedef SListView< TSharedPtr< FGameplayCueEditorHandlerItem > > SGameplayCueHandlerListView;
+	typedef TArray< TSharedPtr< FGameplayCueEditorHandlerItem > > FGameplayCueHandlerArray;
+
+	/** An item in the GameplayCue list */
+	struct FGameplayCueEditorItem
+	{
+		FGameplayTag GameplayCueTagName;
+
+		FGameplayCueHandlerArray HandlerItems;
+
+		TSharedPtr< SGameplayCueHandlerListView > GameplayCueHandlerListView;
+
+		FReply OnAddNewClicked();
+		void OnNewClassPicked(FString SelectedPath);
+	};
+
+	typedef SListView< TSharedPtr< FGameplayCueEditorItem > > SGameplayCueListView;
+	typedef TArray< TSharedPtr< FGameplayCueEditorItem > > FGameplayCueArray;
+	
+
+	TSharedRef<ITableRow> OnGenerateWidgetForGameplayCueListView(TSharedPtr< FGameplayCueEditorItem > InItem, const TSharedRef<STableViewBase>& OwnerTable);
+
+	TSharedPtr< SGameplayCueListView > GameplayCueListView;
+	FGameplayCueArray GameplayCueListItems;
+
+	void UpdateGameplayCueListItems();
+
+	FReply BuildEventMap();
+	
+	TMultiMap<FGameplayTag, UFunction*> EventMap;
+
+	FString AutoSelectTag;
+	FString SearchString;
+};
diff --git a/Engine/Source/Editor/GameplayAbilitiesEditor/Public/GameplayAbilitiesEditorModule.h b/Engine/Source/Editor/GameplayAbilitiesEditor/Public/GameplayAbilitiesEditorModule.h
--- a/Engine/Source/Editor/GameplayAbilitiesEditor/Public/GameplayAbilitiesEditorModule.h
+++ b/Engine/Source/Editor/GameplayAbilitiesEditor/Public/GameplayAbilitiesEditorModule.h
@@ -4,6 +4,38 @@
 
 #include "ModuleManager.h"
 
+struct GAMEPLAYABILITIESEDITOR_API FGameplayCueEditorStrings
+{
+	FString GameplayCueNotifyDescription1;
+	FString GameplayCueNotifyDescription2;
+
+	FString GameplayCueEventDescription1;
+	FString GameplayCueEventDescription2;
+
+	FGameplayCueEditorStrings()
+	{
+		GameplayCueNotifyDescription1 = FString(TEXT("GameplayCue Notifies are stand alone handlers, similiar to AnimNotifies. Most GameplyCues can be implemented through these notifies. Notifies excel at handling standardized effects. The classes below provide the most common functionality needed."));
+		GameplayCueNotifyDescription2 = FString(TEXT(""));
+
+		GameplayCueEventDescription1 = FString(TEXT("GameplayCues can also be implemented via custom events on character blueprints."));
+		GameplayCueEventDescription2 = FString(TEXT("To add a custom BP event, open the blueprint and look for custom events starting with GameplayCue.*"));
+	}
+
+	FGameplayCueEditorStrings(FString Notify1, FString Notify2, FString Event1, FString Event2)
+	{
+		GameplayCueNotifyDescription1 = Notify1;
+		GameplayCueNotifyDescription2 = Notify2;
+		GameplayCueEventDescription1 = Event1;
+		GameplayCueEventDescription2 = Event2;
+	}
+};
+
+DECLARE_DELEGATE_OneParam(FGetGameplayCueNotifyClasses, TArray<UClass*>&);
+DECLARE_DELEGATE_OneParam(FGetGameplayCueInterfaceClasses, TArray<UClass*>&);
+DECLARE_DELEGATE_RetVal_OneParam(FString, FGetGameplayCuePath, FString)
+
+DECLARE_DELEGATE_RetVal(FGameplayCueEditorStrings, FGetGameplayCueEditorStrings);
+
 /**
  * The public interface to this module
  */
@@ -32,4 +64,16 @@ public:
 	{
 		return FModuleManager::Get().IsModuleLoaded( "GameplayAbilitiesEditor" );
 	}
+
+	/** Sets delegate that will be called to retrieve list of gameplay cue notify classes to be presented by GameplayCue Editor when creating a new notify */
+	virtual FGetGameplayCueNotifyClasses& GetGameplayCueNotifyClassesDelegate() = 0;
+
+	virtual FGetGameplayCueInterfaceClasses& GetGameplayCueInterfaceClassesDelegate() = 0;
+
+	/** Sets delegate that will be called to get the save path for a gameplay cue notify that is created through the GameplayCue Editor */
+	virtual FGetGameplayCuePath& GetGameplayCueNotifyPathDelegate() = 0;
+
+	/** Returns strings used in the GameplayCue Editor widgets. Useful for games to override with game specific information for designers (real examples, etc). */
+	virtual FGetGameplayCueEditorStrings& GetGameplayCueEditorStringsDelegate() = 0;
+	
 };
diff --git a/Engine/Source/Runtime/GameplayAbilities/Private/GameplayAbilitiesModule.cpp b/Engine/Source/Runtime/GameplayAbilities/Private/GameplayAbilitiesModule.cpp
--- a/Engine/Source/Runtime/GameplayAbilities/Private/GameplayAbilitiesModule.cpp
+++ b/Engine/Source/Runtime/GameplayAbilities/Private/GameplayAbilitiesModule.cpp
@@ -29,6 +29,11 @@ class FGameplayAbilitiesModule : public IGameplayAbilitiesModule
 		return AbilitySystemGlobals;
 	}
 
+	virtual bool IsAbilitySystemGlobalsAvailable() override
+	{
+		return AbilitySystemGlobals != nullptr;
+	}
+
 	UAbilitySystemGlobals *AbilitySystemGlobals;
 
 	void GetActiveAbilitiesDebugDataForActor(AActor* Actor, FString& AbilityString, bool& bIsUsingAbilities);
diff --git a/Engine/Source/Runtime/GameplayAbilities/Private/GameplayCueManager.cpp b/Engine/Source/Runtime/GameplayAbilities/Private/GameplayCueManager.cpp
--- a/Engine/Source/Runtime/GameplayAbilities/Private/GameplayCueManager.cpp
+++ b/Engine/Source/Runtime/GameplayAbilities/Private/GameplayCueManager.cpp
@@ -340,6 +340,8 @@ void UGameplayCueManager::HandleAssetAdded(UObject *Object)
 
 				check(GlobalCueSet);
 				GlobalCueSet->AddCues(CuesToAdd);
+
+				OnGameplayCueNotifyAddOrRemove.Broadcast();
 			}
 		}
 	}
@@ -367,6 +369,8 @@ void UGameplayCueManager::HandleAssetDeleted(UObject *Object)
 		StringRefs.Add(StringRefToRemove);
 		check(GlobalCueSet);
 		GlobalCueSet->RemoveCuesByStringRefs(StringRefs);
+
+		OnGameplayCueNotifyAddOrRemove.Broadcast();
 	}
 }
 
diff --git a/Engine/Source/Runtime/GameplayAbilities/Public/GameplayAbilitiesModule.h b/Engine/Source/Runtime/GameplayAbilities/Public/GameplayAbilitiesModule.h
--- a/Engine/Source/Runtime/GameplayAbilities/Public/GameplayAbilitiesModule.h
+++ b/Engine/Source/Runtime/GameplayAbilities/Public/GameplayAbilitiesModule.h
@@ -37,6 +37,8 @@ public:
 
 	virtual UAbilitySystemGlobals* GetAbilitySystemGlobals() = 0;
 
+	virtual bool IsAbilitySystemGlobalsAvailable() = 0;
+
 	virtual void GetActiveAbilitiesDebugDataForActor(AActor* Actor, FString& AbilityString, bool& bIsUsingAbilities)=0;
 };
 
diff --git a/Engine/Source/Runtime/GameplayAbilities/Public/GameplayCueManager.h b/Engine/Source/Runtime/GameplayAbilities/Public/GameplayCueManager.h
--- a/Engine/Source/Runtime/GameplayAbilities/Public/GameplayCueManager.h
+++ b/Engine/Source/Runtime/GameplayAbilities/Public/GameplayCueManager.h
@@ -101,6 +101,9 @@ struct GAMEPLAYABILITIES_API FScopedGameplayCueSendContext
 	~FScopedGameplayCueSendContext();
 };
 
+/** Delegate for when GC notifies are added or removed from manager */
+DECLARE_MULTICAST_DELEGATE(FOnGameplayCueNotifyChange);
+
 /**
  *	A self contained handler of a GameplayCue. These are similar to AnimNotifies in implementation.
  */
@@ -189,6 +192,8 @@ class GAMEPLAYABILITIES_API UGameplayCueManager : public UDataAsset
 	bool RegisteredEditorCallbacks;
 
 	bool bAccelerationMapOutdated;
+
+	FOnGameplayCueNotifyChange	OnGameplayCueNotifyAddOrRemove;
 #endif
 
 protected:
diff --git a/Engine/Source/Runtime/GameplayAbilities/Public/GameplayCueSet.h b/Engine/Source/Runtime/GameplayAbilities/Public/GameplayCueSet.h
--- a/Engine/Source/Runtime/GameplayAbilities/Public/GameplayCueSet.h
+++ b/Engine/Source/Runtime/GameplayAbilities/Public/GameplayCueSet.h
@@ -69,10 +69,9 @@ class GAMEPLAYABILITIES_API UGameplayCueSet : public UDataAsset
 	/** Maps GameplayCue Tag to index into above GameplayCues array. */
 	TMap<FGameplayTag, int32> GameplayCueDataMap;
 
+	static FGameplayTag	BaseGameplayCueTag();
+
 protected:
 	virtual bool HandleGameplayCueNotify_Internal(AActor* TargetActor, int32 DataIdx, EGameplayCueEvent::Type EventType, FGameplayCueParameters Parameters);
-	virtual void BuildAccelerationMap_Internal();
-	
-	static FGameplayTag	BaseGameplayCueTag();
-	
+	virtual void BuildAccelerationMap_Internal();	
 };
\ No newline at end of file
diff --git a/Engine/Source/Runtime/GameplayTags/Private/GameplayTagsManager.cpp b/Engine/Source/Runtime/GameplayTags/Private/GameplayTagsManager.cpp
--- a/Engine/Source/Runtime/GameplayTags/Private/GameplayTagsManager.cpp
+++ b/Engine/Source/Runtime/GameplayTags/Private/GameplayTagsManager.cpp
@@ -116,17 +116,27 @@ void UGameplayTagsManager::ConstructGameplayTagTree()
 	if (!GameplayRootTag.IsValid())
 	{
 		GameplayRootTag = MakeShareable(new FGameplayTagNode());
-		
-		for (auto It(GameplayTagTables.CreateIterator()); It; It++)
 		{
-			if (*It)
+#if STATS
+			FString PerfMessage = FString::Printf(TEXT("UGameplayTagsManager::ConstructGameplayTagTree: Construct from data asset"));
+			SCOPE_LOG_TIME_IN_SECONDS(*PerfMessage, nullptr)
+#endif
+		
+			for (auto It(GameplayTagTables.CreateIterator()); It; It++)
 			{
-				PopulateTreeFromDataTable(*It);
+				if (*It)
+				{
+					PopulateTreeFromDataTable(*It);
+				}
 			}
 		}
 
 		if (ShouldImportTagsFromINI())
 		{
+#if STATS
+			FString PerfMessage = FString::Printf(TEXT("UGameplayTagsManager::ConstructGameplayTagTree: ImportINI"));
+			SCOPE_LOG_TIME_IN_SECONDS(*PerfMessage, nullptr)
+#endif
 			// Load any GameplayTagSettings from config (their default object)
 			for (TObjectIterator<UClass> ClassIt; ClassIt; ++ClassIt)
 			{
@@ -150,10 +160,20 @@ void UGameplayTagsManager::ConstructGameplayTagTree()
 
 		if (ShouldUseFastReplication())
 		{
+#if STATS
+			FString PerfMessage = FString::Printf(TEXT("UGameplayTagsManager::ConstructGameplayTagTree: Reconstruct NetIndex"));
+			SCOPE_LOG_TIME_IN_SECONDS(*PerfMessage, nullptr)
+#endif
 			ConstructNetIndex();
 		}
 
-		GameplayTagTreeChangedEvent.Broadcast();
+		{
+#if STATS
+			FString PerfMessage = FString::Printf(TEXT("UGameplayTagsManager::ConstructGameplayTagTree: GameplayTagTreeChangedEvent.Broadcast"));
+			SCOPE_LOG_TIME_IN_SECONDS(*PerfMessage, nullptr)
+#endif
+			GameplayTagTreeChangedEvent.Broadcast();
+		}
 	}
 }
 
