Fix for gameplay abilities linking

diff --git a/Engine/Source/Editor/GameplayAbilitiesEditor/GameplayAbilitiesEditor.Build.cs b/Engine/Source/Editor/GameplayAbilitiesEditor/GameplayAbilitiesEditor.Build.cs
--- a/Engine/Source/Editor/GameplayAbilitiesEditor/GameplayAbilitiesEditor.Build.cs
+++ b/Engine/Source/Editor/GameplayAbilitiesEditor/GameplayAbilitiesEditor.Build.cs
@@ -41,6 +41,7 @@ namespace UnrealBuildTool.Rules
 					"MainFrame",
 					"UnrealEd",
                     "EditorWidgets",
+                    "SourceControl",
 				}
 			);
 		}
