Fix non-unity

diff --git a/Engine/Source/Editor/GameplayAbilitiesEditor/Private/SGameplayCueEditor.cpp b/Engine/Source/Editor/GameplayAbilitiesEditor/Private/SGameplayCueEditor.cpp
--- a/Engine/Source/Editor/GameplayAbilitiesEditor/Private/SGameplayCueEditor.cpp
+++ b/Engine/Source/Editor/GameplayAbilitiesEditor/Private/SGameplayCueEditor.cpp
@@ -13,6 +13,11 @@
 #include "GameplayTagsSettings.h"
 #include "SNotificationList.h"
 #include "NotificationManager.h"
+#include "SSearchBox.h"
+#include "GameplayTagsModule.h"
+#include "AssetEditorManager.h"
+#include "AbilitySystemGlobals.h"
+#include "AssetToolsModule.h"
 
 #define LOCTEXT_NAMESPACE "SGameplayCueEditor"
 
diff --git a/Engine/Source/Editor/GameplayAbilitiesEditor/Private/SGameplayCueEditor.h b/Engine/Source/Editor/GameplayAbilitiesEditor/Private/SGameplayCueEditor.h
--- a/Engine/Source/Editor/GameplayAbilitiesEditor/Private/SGameplayCueEditor.h
+++ b/Engine/Source/Editor/GameplayAbilitiesEditor/Private/SGameplayCueEditor.h
@@ -3,6 +3,7 @@
 #pragma once
 
 #include "SCompoundWidget.h"
+#include "GameplayTagContainer.h"
 
 /** Main CollisionAnalyzer UI widget */
 class SGameplayCueEditor : public SCompoundWidget
diff --git a/Engine/Source/Editor/Sequencer/Private/SequencerTrackLaneFactory.cpp b/Engine/Source/Editor/Sequencer/Private/SequencerTrackLaneFactory.cpp
--- a/Engine/Source/Editor/Sequencer/Private/SequencerTrackLaneFactory.cpp
+++ b/Engine/Source/Editor/Sequencer/Private/SequencerTrackLaneFactory.cpp
@@ -2,6 +2,7 @@
 
 #include "SequencerPrivatePCH.h"
 
+#include "Sequencer.h"
 #include "SequencerTrackLaneFactory.h"
 #include "SSequencerTrackArea.h"
 #include "SSequencerTrackOutliner.h"
diff --git a/Engine/Source/Runtime/GameplayAbilities/Public/GameplayCueSet.h b/Engine/Source/Runtime/GameplayAbilities/Public/GameplayCueSet.h
--- a/Engine/Source/Runtime/GameplayAbilities/Public/GameplayCueSet.h
+++ b/Engine/Source/Runtime/GameplayAbilities/Public/GameplayCueSet.h
@@ -2,6 +2,7 @@
 
 #pragma once
 
+#include "GameplayEffectTypes.h"
 #include "GameplayCueSet.generated.h"
 
 USTRUCT()
