GameplayCue Editor improvement pass:

* I switched this to use a TreeView instead of nested ListViews.
   * This fixes the scrolling bug that Dave mentioned where some items were missing until you scrolled.
   * Currently automatically expand any items that have handlers. Do we want that?
   * We could organize the tags hierarchically. Would that be useful?
* Instead of double-clicking, handlers are now Hyperlinks for more obvious navigation.
* Added search button to suggest that you have to hit [Enter] or click the button for search to actually take place.
* Search highlights matching substring as a reminder that a search is active (i.e. you are not seeing all the items).
* Known Issue : Right now after clicking [Add New] to add a handler, you have to press [Search] before it to shows up.
* Now uses SUserWidget pattern

--------
Integrated using branch UE4-To-UE4-Orion (reversed) of change#2586937 by [REDACTED_PERSON] on 2015/06/13 14:43:29.

diff --git a/Engine/Source/Editor/GameplayAbilitiesEditor/Private/SGameplayCueEditor.cpp b/Engine/Source/Editor/GameplayAbilitiesEditor/Private/SGameplayCueEditor.cpp
--- a/Engine/Source/Editor/GameplayAbilitiesEditor/Private/SGameplayCueEditor.cpp
+++ b/Engine/Source/Editor/GameplayAbilitiesEditor/Private/SGameplayCueEditor.cpp
@@ -13,6 +13,7 @@
 #include "GameplayTagsSettings.h"
 #include "SNotificationList.h"
 #include "NotificationManager.h"
+#include "SHyperlink.h"
 #include "SSearchBox.h"
 #include "GameplayTagsModule.h"
 #include "AssetEditorManager.h"
@@ -22,877 +23,1034 @@
 #define LOCTEXT_NAMESPACE "SGameplayCueEditor"
 
 
-BEGIN_SLATE_FUNCTION_BUILD_OPTIMIZATION
+/** Widget for picking a new GameplayCue Notify class (similiar to actor class picker)  */
+class SGameplayCuePickerDialog : public SCompoundWidget
+{
+public:
+	SLATE_BEGIN_ARGS( SGameplayCuePickerDialog )
+	{}
+
+		SLATE_ARGUMENT( TSharedPtr<SWindow>, ParentWindow )
+		SLATE_ARGUMENT( TArray<UClass*>, DefaultClasses )
+		SLATE_ARGUMENT( FString, GameplayCueTag )
+
+	SLATE_END_ARGS()
+
+	/** Constructs this widget with InArgs */
+	void Construct(const FArguments& InArgs);
+
+	static bool PickGameplayCue(const FText& TitleText, const TArray<UClass*>& DefaultClasses, UClass*& OutChosenClass, FString InGameplayCueTag);
+
+private:
+	/** Handler for when a class is picked in the class picker */
+	void OnClassPicked(UClass* InChosenClass);
+
+	/** Creates the default class widgets */
+	TSharedRef<ITableRow> GenerateListRow(UClass* InItem, const TSharedRef<STableViewBase>& OwnerTable);
+
+	/** Handler for when "Ok" we selected in the class viewer */
+	FReply OnClassPickerConfirmed();
+
+	FReply OnDefaultClassPicked(UClass* InChosenClass);
+
+	/** Handler for the custom button to hide/unhide the default class viewer */
+	void OnDefaultAreaExpansionChanged(bool bExpanded);
+
+	/** Handler for the custom button to hide/unhide the class viewer */
+	void OnCustomAreaExpansionChanged(bool bExpanded);	
+
+private:
+	/** A pointer to the window that is asking the user to select a parent class */
+	TWeakPtr<SWindow> WeakParentWindow;
+
+	/** The class that was last clicked on */
+	UClass* ChosenClass;
+
+	/** A flag indicating that Ok was selected */
+	bool bPressedOk;
+
+	/**  An array of default classes used in the dialog **/
+	TArray<UClass*> DefaultClasses;
+
+	FString GameplayCueTag;
+};
 
-/**
- *	Constructs the main widget for the GameplayCueEditor
- */
-void SGameplayCueEditor::Construct(const FArguments& InArgs)
+BEGIN_SLATE_FUNCTION_BUILD_OPTIMIZATION
+void SGameplayCuePickerDialog::Construct(const FArguments& InArgs)
 {
-	bShowAll = true;
-	bool CanAddFromINI = UGameplayTagsManager::ShouldImportTagsFromINI(); // We only support adding new tags to the ini files.
+	WeakParentWindow = InArgs._ParentWindow;
+	DefaultClasses = InArgs._DefaultClasses;
+	GameplayCueTag = InArgs._GameplayCueTag;
+
+	FLinearColor AssetColor = FLinearColor::White;
+
+	bPressedOk = false;
+	ChosenClass = NULL;
+
+	FString PathStr = SGameplayCueEditor::GetPathNameForGameplayCueTag(GameplayCueTag);
+	FGameplayCueEditorStrings Strings;
+	auto del = IGameplayAbilitiesEditorModule::Get().GetGameplayCueEditorStringsDelegate();
+	if (del.IsBound())
+	{
+		Strings = del.Execute();
+	}
+
 	ChildSlot
 	[
-		SNew(SVerticalBox)
-
-		+SVerticalBox::Slot()
-		.AutoHeight( )
+		SNew(SBorder)
+		.Visibility(EVisibility::Visible)
+		.BorderImage(FEditorStyle::GetBrush("Menu.Background"))
 		[
-			SNew( SHorizontalBox )
-			+SHorizontalBox::Slot()
-			.Padding( 2.0f, 2.0f )
-			.AutoWidth()
+			SNew(SBox)
+			.Visibility(EVisibility::Visible)
+			.Padding(2.f)
+			.WidthOverride(520.0f)
 			[
-					
-				SNew(SButton)
-				.Text(LOCTEXT("SearchBPEvents", "Search BP Events"))
-				.OnClicked( this, &SGameplayCueEditor::BuildEventMap )
-			]
+				SNew(SVerticalBox)
+				+SVerticalBox::Slot()
+				.Padding(2.f, 2.f)
+				.AutoHeight()
+				[
+					SNew(SBorder)
+					.Visibility(EVisibility::Visible)
+					.BorderImage( FEditorStyle::GetBrush("AssetThumbnail.AssetBackground") )
+					.BorderBackgroundColor(AssetColor.CopyWithNewOpacity(0.3f))
+					[
+						SNew(SExpandableArea)
+						.AreaTitle(NSLOCTEXT("SGameplayCuePickerDialog", "CommonClassesAreaTitle", "GameplayCue Notifies"))
+						.BodyContent()
+						[
+							SNew(SVerticalBox)
+							+SVerticalBox::Slot()
+							.Padding(2.f, 2.f)
+							.AutoHeight()
+							[
+								SNew(STextBlock)
+								
+								.Text(FText::FromString(Strings.GameplayCueNotifyDescription1))
+								.AutoWrapText(true)
+							]
 
-			+SHorizontalBox::Slot()
-			.Padding( 2.0f, 2.0f )
-			.AutoWidth()
-			[
-				// Check box that controls LIVE MODE
-				SNew(SCheckBox)
-				.IsChecked(this, &SGameplayCueEditor::HandleShowAllCheckBoxIsChecked)
-				.OnCheckStateChanged(this, &SGameplayCueEditor::HandleShowAllCheckedStateChanged)
+							+SVerticalBox::Slot()
+							.AutoHeight( )
+							[
+								SNew(SListView < UClass*  >)
+								.ItemHeight(48)
+								.SelectionMode(ESelectionMode::None)
+								.ListItemsSource(&DefaultClasses)
+								.OnGenerateRow(this, &SGameplayCuePickerDialog::GenerateListRow)
+							]
+
+							+SVerticalBox::Slot()
+							.Padding(2.f, 2.f)
+							.AutoHeight()
+							[
+								SNew(STextBlock)
+								.Text(FText::FromString(FString::Printf(TEXT("This will create a new GameplayCue Notify here:"))))
+								.AutoWrapText(true)
+							]
+							+SVerticalBox::Slot()
+							.Padding(2.f, 2.f)
+							.AutoHeight()
+							[
+								SNew(STextBlock)
+								.Text(FText::FromString(PathStr))
+								.HighlightText(FText::FromString(PathStr))
+								.AutoWrapText(true)
+							]
+							+SVerticalBox::Slot()
+							.Padding(2.f, 2.f)
+							.AutoHeight()
+							[
+								SNew(STextBlock)
+								.Text(FText::FromString(Strings.GameplayCueNotifyDescription2))
+								.AutoWrapText(true)
+							]
+						]
+					]
+				]
+
+				+SVerticalBox::Slot()
+				.Padding(2.f, 2.f)
+				.AutoHeight()
 				[
-					SNew(STextBlock)
-					.Text(LOCTEXT("HieUnhandled", "Hide Unhandled GameplayCues"))
+					SNew(SBorder)
+					.Visibility(EVisibility::Visible)
+					.BorderImage( FEditorStyle::GetBrush("AssetThumbnail.AssetBackground") )
+					.BorderBackgroundColor(AssetColor.CopyWithNewOpacity(0.3f))
+					[
+						SNew(SExpandableArea)
+						.AreaTitle(NSLOCTEXT("SGameplayCuePickerDialogEvents", "CommonClassesAreaTitleEvents", "Custom BP Events"))
+						.BodyContent()
+						[
+							SNew(SVerticalBox)
+							+SVerticalBox::Slot()
+							.Padding(2.f, 2.f)
+							.AutoHeight()
+							[
+								SNew(STextBlock)
+								.Text(FText::FromString(Strings.GameplayCueEventDescription1))
+								.AutoWrapText(true)
+
+							]
+
+							+SVerticalBox::Slot()
+							.Padding(2.f, 2.f)
+							.AutoHeight()
+							[
+								SNew(STextBlock)
+								.Text(FText::FromString(Strings.GameplayCueEventDescription2))
+								.AutoWrapText(true)
+							]							
+						]
+					]
 				]
 			]
 		]
+	];
+}
+END_SLATE_FUNCTION_BUILD_OPTIMIZATION
+
+
+static const FName CueTagColumnName("GameplayCueTags");
+static const FName CueHandlerColumnName("GameplayCueHandlers");
+
+/** Spawns window for picking new GameplayCue handler/notify */
+bool SGameplayCuePickerDialog::PickGameplayCue(const FText& TitleText, const TArray<UClass*>& DefaultClasses, UClass*& OutChosenClass, FString InGameplayCueName)
+{
+	// Create the window to pick the class
+	TSharedRef<SWindow> PickerWindow = SNew(SWindow)
+		.Title(TitleText)
+		.SizingRule( ESizingRule::Autosized )
+		.ClientSize( FVector2D( 0.f, 600.f ))
+		.SupportsMaximize(false)
+		.SupportsMinimize(false);
 
+	TSharedRef<SGameplayCuePickerDialog> ClassPickerDialog = SNew(SGameplayCuePickerDialog)
+		.ParentWindow(PickerWindow)
+		.DefaultClasses(DefaultClasses)
+		.GameplayCueTag(InGameplayCueName);
+
+	PickerWindow->SetContent(ClassPickerDialog);
+
+	GEditor->EditorAddModalWindow(PickerWindow);
+
+	if (ClassPickerDialog->bPressedOk)
+	{
+		OutChosenClass = ClassPickerDialog->ChosenClass;
+		return true;
+	}
+	else
+	{
+		// Ok was not selected, NULL the class
+		OutChosenClass = NULL;
+		return false;
+	}
+}
+
+void SGameplayCuePickerDialog::OnClassPicked(UClass* InChosenClass)
+{
+	ChosenClass = InChosenClass;
+}
+
+/** Generates rows in the list of GameplayCueNotify classes to pick from */
+TSharedRef<ITableRow> SGameplayCuePickerDialog::GenerateListRow(UClass* ItemClass, const TSharedRef<STableViewBase>& OwnerTable)
+{	
+	const FSlateBrush* ItemBrush = FClassIconFinder::FindIconForClass(ItemClass);
+
+	return 
+	SNew(STableRow< UClass* >, OwnerTable)
+	[
+		SNew(SVerticalBox)
 		+SVerticalBox::Slot()
-		.AutoHeight( )
+		.MaxHeight(60.0f)
+		.Padding(10.0f, 6.0f, 0.0f, 4.0f)
 		[
-			SNew( SHorizontalBox )
+			SNew(SHorizontalBox)
 			+SHorizontalBox::Slot()
-			.Padding( 2.0f, 2.0f )
-			.AutoWidth()
+			.FillWidth(0.45f)
 			[
-				SAssignNew(NewGameplayCueTextBox, SEditableTextBox)
-				.MinDesiredWidth(210.0f)
-				.HintText(LOCTEXT("GameplayCueXY", "GameplayCue.X.Y"))
-				.OnTextCommitted( this, &SGameplayCueEditor::OnNewGameplayCueTagCommited )
-			]
+				SNew(SButton)
+				.OnClicked(this, &SGameplayCuePickerDialog::OnDefaultClassPicked, ItemClass)
+				.Content()
+				[
+					SNew(SHorizontalBox)
+					+SHorizontalBox::Slot()
+					.HAlign(HAlign_Center)
+					.VAlign(VAlign_Center)
+					.FillWidth(0.12f)
+					[
+						SNew(SImage)
+						.Image(ItemBrush)
+					]
+					+SHorizontalBox::Slot()
+					.VAlign(VAlign_Center)
+					.Padding(4.0f, 0.0f)
+					.FillWidth(0.8f)
+					[
+						SNew(STextBlock)
+						.Text(ItemClass->GetDisplayNameText())
+					]
 
+				]
+			]
 			+SHorizontalBox::Slot()
-			.Padding( 2.0f, 2.0f )
-			.AutoWidth()
+			.Padding(10.0f, 0.0f)
 			[
-				SNew(SButton)
-				.Text(LOCTEXT("AddNew", "Add New"))
-				.OnClicked( this, &SGameplayCueEditor::OnNewGameplayCueButtonPressed )
-				.Visibility( CanAddFromINI ? EVisibility::Visible : EVisibility::Collapsed )
+				SNew(STextBlock)
+				.Text(ItemClass->GetToolTipText(true))
+				.AutoWrapText(true)
 			]
-		]
-
-		+SVerticalBox::Slot()
-		.AutoHeight( )
-		[
-			SNew( SHorizontalBox )
 			+SHorizontalBox::Slot()
-			.Padding( 2.0f, 2.0f )
 			.AutoWidth()
 			[
-				SNew(SSearchBox)
-				.MinDesiredWidth(210.0f)				
-				.OnTextCommitted( this, &SGameplayCueEditor::OnSearchTagCommited )
-			]
-		]
-		
-		+SVerticalBox::Slot()
-		.FillHeight( 1.0f )
-		[
-			SNew( SBorder )
-			.BorderImage( FEditorStyle::GetBrush( "ToolBar.Background" ) )
-			[
-				SAssignNew( GameplayCueListView, SGameplayCueListView )
-					.ItemHeight( 24 )
-					.ListItemsSource( &GameplayCueListItems )
-					.OnGenerateRow( this, &SGameplayCueEditor::OnGenerateWidgetForGameplayCueListView )
-					.HeaderRow
-					(
-						SNew(SHeaderRow)
-						+SHeaderRow::Column("GameplayCueTags")
-						.DefaultLabel(NSLOCTEXT("GameplayCueEditor", "GameplayCueTag", "Tag"))
-						.FillWidth(0.25)
-
-						+SHeaderRow::Column("GameplayCueHandlers")
-						.DefaultLabel(NSLOCTEXT("GameplayCueEditor", "GameplayCueHandlers", "Handlers"))
-						//.FillWidth(1000)
-					)
+				FEditorClassUtils::GetDocumentationLinkWidget(ItemClass)
 			]
 		]
 	];
-
-	UpdateGameplayCueListItems();
 }
 
-END_SLATE_FUNCTION_BUILD_OPTIMIZATION
+FReply SGameplayCuePickerDialog::OnDefaultClassPicked(UClass* InChosenClass)
+{
+	ChosenClass = InChosenClass;
+	bPressedOk = true;
+	if (WeakParentWindow.IsValid())
+	{
+		WeakParentWindow.Pin()->RequestDestroyWindow();
+	}
+	return FReply::Handled();
+}
 
-SGameplayCueEditor::~SGameplayCueEditor()
+FReply SGameplayCuePickerDialog::OnClassPickerConfirmed()
 {
-	
+	if (ChosenClass == NULL)
+	{
+		FMessageDialog::Open(EAppMsgType::Ok, NSLOCTEXT("EditorFactories", "MustChooseClassWarning", "You must choose a class."));
+	}
+	else
+	{
+		bPressedOk = true;
+
+		if (WeakParentWindow.IsValid())
+		{
+			WeakParentWindow.Pin()->RequestDestroyWindow();
+		}
+	}
+	return FReply::Handled();
 }
 
-void SGameplayCueEditor::OnSearchTagCommited(const FText& InText, ETextCommit::Type InCommitType)
+FString SGameplayCueEditor::GetPathNameForGameplayCueTag(FString GameplayCueTagName)
 {
-	if (InCommitType == ETextCommit::OnEnter || InCommitType == ETextCommit::OnCleared)
+	FString NewDefaultPathName = FString::Printf(TEXT("/GC_%s"), *GameplayCueTagName);
+	auto PathDel = IGameplayAbilitiesEditorModule::Get().GetGameplayCueNotifyPathDelegate();
+	if (PathDel.IsBound())
 	{
-		SearchString = InText.ToString();
-		UpdateGameplayCueListItems();
+		NewDefaultPathName = PathDel.Execute(GameplayCueTagName);
 	}
+	NewDefaultPathName.ReplaceInline(TEXT("."), TEXT("_"));
+	return NewDefaultPathName;
 }
 
-void SGameplayCueEditor::OnNewGameplayCueTagCommited(const FText& InText, ETextCommit::Type InCommitType)
+
+class SGameplayCueEditorImpl : public SGameplayCueEditor
 {
-	// Only support adding tags via ini file
-	if (UGameplayTagsManager::ShouldImportTagsFromINI() == false)
+public:
+
+	virtual void Construct(const FArguments& InArgs) override;
+
+	virtual void OnNewGameplayCueTagCommited(const FText& InText, ETextCommit::Type InCommitType) override
+	{
+		// Only support adding tags via ini file
+		if (UGameplayTagsManager::ShouldImportTagsFromINI() == false)
+		{
+			return;
+		}
+
+		if (InCommitType == ETextCommit::OnEnter)
+		{
+			CreateNewGameplayCueTag();
+		}
+	}
+
+	virtual void OnSearchTagCommited(const FText& InText, ETextCommit::Type InCommitType) override
 	{
-		return;
+		if (InCommitType == ETextCommit::OnEnter || InCommitType == ETextCommit::OnCleared || InCommitType == ETextCommit::OnUserMovedFocus)
+		{
+			SearchText = InText;
+			UpdateGameplayCueListItems();
+		}
 	}
 
-	if (InCommitType == ETextCommit::OnEnter)
+	FReply DoSearch()
 	{
-		CreateNewGameplayCueTag();
+		UpdateGameplayCueListItems();
+		return FReply::Handled();
 	}
-}
 
-FReply SGameplayCueEditor::OnNewGameplayCueButtonPressed()
-{
-	CreateNewGameplayCueTag();
-	return FReply::Handled();
-}
 
-/**
- *	Checks out config file, adds new tag, repopulates widget cue list
- */
-void SGameplayCueEditor::CreateNewGameplayCueTag()
-{
-	FScopedSlowTask SlowTask(0.f, LOCTEXT("AddingNewGameplaycue", "Adding new GameplayCue Tag"));
-	SlowTask.MakeDialog();
-				
+	virtual FReply OnNewGameplayCueButtonPressed() override
+	{
+		CreateNewGameplayCueTag();
+		return FReply::Handled();
+	}
 
-	FString str = NewGameplayCueTextBox->GetText().ToString();
+private:
 
-	AutoSelectTag = str;
-	
-	UGameplayTagsSettings* Settings = GetMutableDefault<UGameplayTagsSettings>();
-	
-	if (Settings)
+	TSharedPtr<SEditableTextBox>	NewGameplayCueTextBox;
+
+	/**
+	*	Checks out config file, adds new tag, repopulates widget cue list
+	*/
+	void CreateNewGameplayCueTag()
 	{
-		
-		FString RelativeConfigFilePath = Settings->GetDefaultConfigFilename();
-		FString ConfigPath = FPaths::ConvertRelativePathToFull(RelativeConfigFilePath);
+		FScopedSlowTask SlowTask(0.f, LOCTEXT("AddingNewGameplaycue", "Adding new GameplayCue Tag"));
+		SlowTask.MakeDialog();
+
+
+		FString str = NewGameplayCueTextBox->GetText().ToString();
 
-		if (ISourceControlModule::Get().IsEnabled())
+		AutoSelectTag = str;
+
+		UGameplayTagsSettings* Settings = GetMutableDefault<UGameplayTagsSettings>();
+
+		if (Settings)
 		{
-			FText ErrorMessage;
 
-			if (!SourceControlHelpers::CheckoutOrMarkForAdd(ConfigPath, FText::FromString(ConfigPath), NULL, ErrorMessage))
+			FString RelativeConfigFilePath = Settings->GetDefaultConfigFilename();
+			FString ConfigPath = FPaths::ConvertRelativePathToFull(RelativeConfigFilePath);
+
+			if (ISourceControlModule::Get().IsEnabled())
 			{
-				FNotificationInfo Info(ErrorMessage);
-				Info.ExpireDuration = 3.0f;
-				FSlateNotificationManager::Get().AddNotification(Info);
+				FText ErrorMessage;
+
+				if (!SourceControlHelpers::CheckoutOrMarkForAdd(ConfigPath, FText::FromString(ConfigPath), NULL, ErrorMessage))
+				{
+					FNotificationInfo Info(ErrorMessage);
+					Info.ExpireDuration = 3.0f;
+					FSlateNotificationManager::Get().AddNotification(Info);
+				}
 			}
-		}
-		else
-		{
-			if (!FPlatformFileManager::Get().GetPlatformFile().SetReadOnly(*ConfigPath, false))
+			else
 			{
-				FText NotificationErrorText = FText::Format(LOCTEXT("FailedToMakeWritable", "Could not make {0} writable."), FText::FromString(ConfigPath));
+				if (!FPlatformFileManager::Get().GetPlatformFile().SetReadOnly(*ConfigPath, false))
+				{
+					FText NotificationErrorText = FText::Format(LOCTEXT("FailedToMakeWritable", "Could not make {0} writable."), FText::FromString(ConfigPath));
 
-				FNotificationInfo Info(NotificationErrorText);
-				Info.ExpireDuration = 3.0f;
+					FNotificationInfo Info(NotificationErrorText);
+					Info.ExpireDuration = 3.0f;
 
-				FSlateNotificationManager::Get().AddNotification(Info);
+					FSlateNotificationManager::Get().AddNotification(Info);
+				}
 			}
-		}
-		
-		Settings->GameplayTags.Add(str);			
-		IGameplayTagsModule::Get().GetGameplayTagsManager().DestroyGameplayTagTree();			
 
-		{
+			Settings->GameplayTags.Add(str);
+			IGameplayTagsModule::Get().GetGameplayTagsManager().DestroyGameplayTagTree();
+
+			{
 #if STATS
-			FString PerfMessage = FString::Printf(TEXT("ConstructGameplayTagTree GameplayTag tables"));
-			SCOPE_LOG_TIME_IN_SECONDS(*PerfMessage, nullptr)
+				FString PerfMessage = FString::Printf(TEXT("ConstructGameplayTagTree GameplayTag tables"));
+				SCOPE_LOG_TIME_IN_SECONDS(*PerfMessage, nullptr)
 #endif
-			IGameplayTagsModule::Get().GetGameplayTagsManager().ConstructGameplayTagTree();
+					IGameplayTagsModule::Get().GetGameplayTagsManager().ConstructGameplayTagTree();
+			}
+
+			Settings->UpdateDefaultConfigFile();
+
+			UpdateGameplayCueListItems();
 		}
 
-		Settings->UpdateDefaultConfigFile();
+		NewGameplayCueTextBox->SetText(FText::GetEmpty());
+	}
 
+	void HandleShowAllCheckedStateChanged(ECheckBoxState NewValue)
+	{
+		bShowAll = (NewValue == ECheckBoxState::Unchecked);
 		UpdateGameplayCueListItems();
 	}
 
-	NewGameplayCueTextBox->SetText(FText::GetEmpty());
-}
+	/** Callback for getting the checked state of the focus check box. */
+	ECheckBoxState HandleShowAllCheckBoxIsChecked() const
+	{
+		return bShowAll ? ECheckBoxState::Unchecked : ECheckBoxState::Checked;
+	}
 
-void SGameplayCueEditor::HandleShowAllCheckedStateChanged( ECheckBoxState NewValue )
-{
-	bShowAll = (NewValue == ECheckBoxState::Unchecked);
-	UpdateGameplayCueListItems();
-}	
+	struct FCueItem;
+	struct FCueHandlerItem;
+	
+	/** Base class for any item in the the Cue/Handler Tree. */
+	struct FTreeItem : public TSharedFromThis<FTreeItem>
+	{
+		virtual ~FTreeItem(){}
+		virtual TSharedPtr<FCueItem> AsCueItem(){ return TSharedPtr<FCueItem>(); }
+		virtual TSharedPtr<FCueHandlerItem> AsCueHandlerItem(){ return TSharedPtr<FCueHandlerItem>(); }
+	};
 
-/** Callback for getting the checked state of the focus check box. */
-ECheckBoxState SGameplayCueEditor::HandleShowAllCheckBoxIsChecked() const
-{
-	return bShowAll ? ECheckBoxState::Unchecked : ECheckBoxState::Checked;
-}
+	struct FCueHandlerItem : public FTreeItem
+	{
+		FStringAssetReference GameplayCueNotifyObj;
+		TWeakObjectPtr<UFunction> FunctionPtr;
 
-/** Builds widget for rows in the GameplayCue Editor tab */
-TSharedRef<ITableRow> SGameplayCueEditor::OnGenerateWidgetForGameplayCueListView(TSharedPtr< FGameplayCueEditorItem > InItem, const TSharedRef<STableViewBase>& OwnerTable)
-{
-	class SGameplayCueItemWidget : public SMultiColumnTableRow< TSharedPtr< FGameplayCueEditorItem > >
+		virtual TSharedPtr<FCueHandlerItem> AsCueHandlerItem() override
+		{
+			return SharedThis(this);
+		}
+	};
+
+	typedef SListView< TSharedPtr< FCueHandlerItem > > SGameplayCueHandlerListView;
+
+	/** An item in the GameplayCue list */
+	struct FCueItem : public FTreeItem
 	{
-		public:
-		SLATE_BEGIN_ARGS( SGameplayCueItemWidget ){}
-		SLATE_END_ARGS()
+		virtual TSharedPtr<FCueItem> AsCueItem(){ return SharedThis(this); }
+
+		FGameplayTag GameplayCueTagName;
 
-		void Construct( const FArguments& InArgs, const TSharedRef<STableViewBase>& InOwnerTable, TSharedPtr<FGameplayCueEditorItem> InListItem )
+		TArray< TSharedPtr< FCueHandlerItem > > HandlerItems;
+
+		TSharedPtr< SGameplayCueHandlerListView > GameplayCueHandlerListView;
+
+		
+		void OnNewClassPicked(FString SelectedPath)
 		{
-			Item = InListItem;
-	
-			SMultiColumnTableRow< TSharedPtr< FGameplayCueEditorItem > >::Construct( FSuperRowType::FArguments(), InOwnerTable );			
+
 		}
+	};
 
-		TSharedRef<SWidget> GenerateWidgetForColumn( const FName& ColumnName )
+	typedef STreeView< TSharedPtr< FTreeItem > > SGameplayCueTreeView;
+
+
+	/** Builds widget for rows in the GameplayCue Editor tab */
+	TSharedRef<ITableRow> OnGenerateWidgetForGameplayCueListView(TSharedPtr< FTreeItem > InItem, const TSharedRef<STableViewBase>& OwnerTable)
+	{
+		class SCueItemWidget : public SMultiColumnTableRow< TSharedPtr< FCueItem > >
 		{
-			if ( ColumnName == "GameplayCueTags" )
+		public:
+			SLATE_BEGIN_ARGS(SCueItemWidget){}
+			SLATE_END_ARGS()
+
+			void Construct(const FArguments& InArgs, const TSharedRef<SGameplayCueTreeView>& InOwnerTable, FSimpleDelegate InRefreshDelegate,  TSharedPtr<FCueItem> InListItem, FText InSearchTerm )
 			{
-				return
-					SNew( STextBlock )
-					.Text( FText::FromString(Item->GameplayCueTagName.ToString()) );
+				Item = InListItem;
+				SearchTerm = InSearchTerm;
+				RefreshDelegate = InRefreshDelegate;
+				SMultiColumnTableRow< TSharedPtr< FCueItem > >::Construct(FSuperRowType::FArguments(), InOwnerTable);
 			}
-			else if ( ColumnName == "GameplayCueHandlers" )
+		private:
+
+			virtual TSharedRef<SWidget> GenerateWidgetForColumn(const FName& ColumnName) override
 			{
-				return 
-				SNew( SVerticalBox )				
-				// Notifies
-				+ SVerticalBox::Slot()
-					.AutoHeight()
-					.Padding( 2.0f, 2.0f )
-				[
-					SAssignNew( Item->GameplayCueHandlerListView, SGameplayCueHandlerListView )
-					.ItemHeight( 24 )
-					.ListItemsSource( &Item->HandlerItems )
-					.OnGenerateRow( this, &SGameplayCueItemWidget::OnGenerateWidgetForGameplayCueHandlerListView )
-				]
-	
-				// Add New button
-				+ SVerticalBox::Slot()
-					.AutoHeight()
-					.Padding( 2.0f, 0.0f )
-				[
-					SNew( SHorizontalBox )
+				if (ColumnName == CueTagColumnName)
+				{
+					return
+					SNew(SHorizontalBox)
 					+ SHorizontalBox::Slot()
 					.AutoWidth()
-					.Padding( 0.0f, 0.0f )
-					[					
-						SNew( SButton )
-							.Text( LOCTEXT("AddNew", "Add New") )
-							.OnClicked( Item.ToSharedRef(), &FGameplayCueEditorItem::OnAddNewClicked )
+					[
+						SNew(SExpanderArrow, SharedThis(this))
 					]
-				];
-			}
-			else
-			{
-				return SNew(STextBlock) .Text(LOCTEXT("UnknownColumn", "Unknown Column"));
-			}
-
-		}
-
-		TSharedRef<ITableRow> OnGenerateWidgetForGameplayCueHandlerListView(TSharedPtr< FGameplayCueEditorHandlerItem > InItem, const TSharedRef<STableViewBase>& OwnerTable)
-		{			
-			FLinearColor Color(0.2f, 0.4f, 0.6f, 1.f);
-			FString ObjName;
-			
-			if ( InItem->GameplayCueNotifyObj.AssetLongPathname.IsEmpty() == false)
-			{
-				ObjName = InItem->GameplayCueNotifyObj.AssetLongPathname;
-			
-				int32 idx;
-				if (ObjName.FindLastChar(TEXT('.'), idx))
+					+ SHorizontalBox::Slot()
+					.FillWidth(1)
+					.VAlign(VAlign_Center)
+					[
+						SNew(STextBlock)
+						.HighlightText(SearchTerm)
+						.Text(FText::FromString(Item->GameplayCueTagName.ToString()))
+					];
+				}
+				else if (ColumnName == CueHandlerColumnName)
 				{
-					ObjName = ObjName.RightChop(idx+1);
-					if (ObjName.FindLastChar(TEXT('_'), idx))
-					{
-						ObjName = ObjName.Left(idx);
-					}				
+					return
+					SNew(SBox)
+					.Padding( FMargin(2.0f, 0.0f) )
+					.HAlign(HAlign_Left)
+					.VAlign(VAlign_Center)
+					[
+						SNew(SButton)
+						.Text(LOCTEXT("AddNew", "Add New"))
+						.OnClicked(this, &SCueItemWidget::OnAddNewClicked)
+					];
 				}
-			}
-			else if (InItem->FunctionPtr.IsValid())
-			{
-				Color = FLinearColor(1.f, 1.f, 1.0f, 1.f);
-				UFunction* Func = InItem->FunctionPtr.Get();
-				UClass* OuterClass = Cast<UClass>(Func->GetOuter());
-				if (OuterClass)
+				else
 				{
-					ObjName = OuterClass->GetName();
-					ObjName.RemoveFromEnd(TEXT("_c"));
+					return SNew(STextBlock).Text(LOCTEXT("UnknownColumn", "Unknown Column"));
 				}
-			}
-
-			return	SNew( STableRow< TSharedPtr<FGameplayCueEditorHandlerItem> >, OwnerTable )
-			.Content()
-			[	
-				SNew( STextBlock )
-				.Text( FText::FromString(ObjName ) )
-				.ColorAndOpacity(Color)
-				.OnDoubleClicked( this, &SGameplayCueItemWidget::OnHandlerDoubleClicked, InItem )
-			];
-		}
 
-		FReply OnHandlerDoubleClicked(TSharedPtr< FGameplayCueEditorHandlerItem > InItem)
-		{
-			if (InItem->GameplayCueNotifyObj.IsValid())
-			{
-				TArray<FString>	AssetList;
-				AssetList.Add(InItem->GameplayCueNotifyObj.ToString());
-				FAssetEditorManager::Get().OpenEditorsForAssets(AssetList);
 			}
-			else if(InItem->FunctionPtr.IsValid())
+
+			/** Create new GameplayCueNotify: brings up dialogue to pick class, then creates it via the content browser. */
+			FReply OnAddNewClicked()
 			{
-				TArray<FString>	AssetList;
-				AssetList.Add(InItem->FunctionPtr->GetOuter()->GetPathName());
-				FAssetEditorManager::Get().OpenEditorsForAssets(AssetList);
-			}
-			
-			return FReply::Handled();
-		}
+				FAssetToolsModule& AssetToolsModule = FModuleManager::GetModuleChecked<FAssetToolsModule>("AssetTools");
+				FContentBrowserModule& ContentBrowserModule = FModuleManager::LoadModuleChecked<FContentBrowserModule>("ContentBrowser");
 
-		TSharedPtr< FGameplayCueEditorItem > Item;
-	};
+				TArray<UClass*>	NotifyClasses;
+				auto del = IGameplayAbilitiesEditorModule::Get().GetGameplayCueNotifyClassesDelegate();
+				del.ExecuteIfBound(NotifyClasses);
+				if (NotifyClasses.Num() == 0)
+				{
+					NotifyClasses.Add(UGameplayCueNotify_Static::StaticClass());
+					NotifyClasses.Add(AGameplayCueNotify_Actor::StaticClass());
+				}
 
-	return SNew( SGameplayCueItemWidget, OwnerTable, InItem );
-}
+				// --------------------------------------------------
 
+				// Null the parent class to ensure one is selected	
 
+				const FText TitleText = LOCTEXT("CreateBlueprintOptions", "New GameplayCue Handler");
+				UClass* ChosenClass = nullptr;
+				const bool bPressedOk = SGameplayCuePickerDialog::PickGameplayCue(TitleText, NotifyClasses, ChosenClass, Item->GameplayCueTagName.ToString());
+				if (bPressedOk)
+				{
+					FString NewDefaultPathName = SGameplayCueEditor::GetPathNameForGameplayCueTag(Item->GameplayCueTagName.ToString());
 
-void SGameplayCueEditor::FGameplayCueEditorItem::OnNewClassPicked(FString SelectedPath)
-{
-	
-}
+					// Make sure the name is unique
+					FString AssetName;
+					FString PackageName;
+					AssetToolsModule.Get().CreateUniqueAssetName(NewDefaultPathName, TEXT(""), /*out*/ PackageName, /*out*/ AssetName);
+					const FString PackagePath = FPackageName::GetLongPackagePath(PackageName);
 
-/** Slow task: use asset registry to find blueprints, load an inspect them to see what GC events they implement. */
-FReply SGameplayCueEditor::BuildEventMap()
-{
-	FScopedSlowTask SlowTask(100.f, LOCTEXT("BuildEventMap", "Searching Blueprints for GameplayCue events"));
-	SlowTask.MakeDialog();
-	SlowTask.EnterProgressFrame(10.f);
+					// Create the GameplayCue Notify
+					UBlueprintFactory* BlueprintFactory = NewObject<UBlueprintFactory>();
+					BlueprintFactory->ParentClass = ChosenClass;
+					ContentBrowserModule.Get().CreateNewAsset(AssetName, PackagePath, UBlueprint::StaticClass(), BlueprintFactory);
 
-	EventMap.Empty();
+					RefreshDelegate.ExecuteIfBound();
+				}
 
-	IGameplayTagsModule& GameplayTagModule = IGameplayTagsModule::Get();
+				return FReply::Handled();
+			}
 
-	auto del = IGameplayAbilitiesEditorModule::Get().GetGameplayCueInterfaceClassesDelegate();
-	if (del.IsBound())
-	{
-		TArray<UClass*> InterfaceClasses;
-		del.ExecuteIfBound(InterfaceClasses);
-		float WorkPerClass = InterfaceClasses.Num() > 0 ? 90.f / static_cast<float>(InterfaceClasses.Num()) : 0.f;
+			TSharedPtr< FCueItem > Item;
+			FText SearchTerm;
+			FSimpleDelegate RefreshDelegate;
+		};
 
-		for (UClass* InterfaceClass : InterfaceClasses)
+		class SCueHandlerItemWidget : public SMultiColumnTableRow < TSharedPtr< FCueHandlerItem > >
 		{
-			SlowTask.EnterProgressFrame(WorkPerClass);
+		public:
+			SLATE_BEGIN_ARGS(SCueHandlerItemWidget){}
+			SLATE_END_ARGS()
 
-			TArray<FAssetData> GameplayCueInterfaceActors;
+			void Construct(const FArguments& InArgs, const TSharedRef<STableViewBase>& InOwnerTable, TSharedPtr<FCueHandlerItem> InListItem)
 			{
-		#if STATS
-				FString PerfMessage = FString::Printf(TEXT("Searched asset registry %s "), *InterfaceClass->GetName());
-				SCOPE_LOG_TIME_IN_SECONDS(*PerfMessage, nullptr)
-		#endif
-
-				UObjectLibrary* ObjLibrary = UObjectLibrary::CreateLibrary(InterfaceClass, true, true);
-				ObjLibrary->LoadBlueprintAssetDataFromPath(TEXT("/Game/"));
-				ObjLibrary->GetAssetDataList(GameplayCueInterfaceActors);
+				CueHandlerItem = InListItem;
+				SMultiColumnTableRow< TSharedPtr< FCueHandlerItem > >::Construct(FSuperRowType::FArguments(), InOwnerTable);
 			}
-			
-			{
-		#if STATS
-				FString PerfMessage = FString::Printf(TEXT("Fully Loaded GameplayCueNotify actors %s "), *InterfaceClass->GetName());
-				SCOPE_LOG_TIME_IN_SECONDS(*PerfMessage, nullptr)
-		#endif				
+		private:
 
-				for (const FAssetData& AssetData : GameplayCueInterfaceActors)
+			virtual TSharedRef<SWidget> GenerateWidgetForColumn(const FName& ColumnName) override
+			{
+				if (ColumnName == CueTagColumnName)
+				{
+					return SNew(SSpacer);
+				}
+				else if (ColumnName == CueHandlerColumnName)
 				{
-					UBlueprint* BP = Cast<UBlueprint>(AssetData.GetAsset());
-					if (BP)
+					if (CueHandlerItem->GameplayCueNotifyObj.AssetLongPathname.IsEmpty() == false)
 					{
-						for (TFieldIterator<UFunction> FuncIt(BP->GeneratedClass, EFieldIteratorFlags::ExcludeSuper); FuncIt; ++FuncIt)
+						FString ObjName = CueHandlerItem->GameplayCueNotifyObj.AssetLongPathname;
+
+						int32 idx;
+						if (ObjName.FindLastChar(TEXT('.'), idx))
 						{
-							FString FuncName = FuncIt->GetName();
-							if (FuncName.Contains("GameplayCue"))
+							ObjName = ObjName.RightChop(idx + 1);
+							if (ObjName.FindLastChar(TEXT('_'), idx))
 							{
-								FuncName.ReplaceInline(TEXT("_"), TEXT("."));
-								FGameplayTag FoundTag = GameplayTagModule.RequestGameplayTag(FName(*FuncName), false);
-								if (FoundTag.IsValid())
-								{
-									EventMap.AddUnique(FoundTag, *FuncIt);
-								}
+								ObjName = ObjName.Left(idx);
 							}
+						}
 
+						return
+						SNew(SBox)
+						.HAlign(HAlign_Left)
+						[
+							SNew(SHyperlink)
+							.Style(FEditorStyle::Get(), "Common.GotoBlueprintHyperlink")
+							.Text(FText::FromString(ObjName))
+							.OnNavigate(this, &SCueHandlerItemWidget::NavigateToHandler)
+						];
+					}
+					else if (CueHandlerItem->FunctionPtr.IsValid())
+					{
+						FString ObjName;
+						UFunction* Func = CueHandlerItem->FunctionPtr.Get();
+						UClass* OuterClass = Cast<UClass>(Func->GetOuter());
+						if (OuterClass)
+						{
+							ObjName = OuterClass->GetName();
+							ObjName.RemoveFromEnd(TEXT("_c"));
 						}
+
+						return
+						SNew(SBox)
+						.HAlign(HAlign_Left)
+						[
+							SNew(SHyperlink)
+							.Text(FText::FromString(ObjName))
+							.OnNavigate(this, &SCueHandlerItemWidget::NavigateToHandler)
+						];
+					}
+					else
+					{
+						return SNew(STextBlock).Text(LOCTEXT("UnknownHandler", "Unknown HandlerType"));
 					}
 				}
+				else
+				{
+					return SNew(STextBlock).Text(LOCTEXT("UnknownColumn", "Unknown Column"));
+				}
 			}
-		}
 
-		UpdateGameplayCueListItems();
-	}
+			void NavigateToHandler()
+			{
+				if (CueHandlerItem->GameplayCueNotifyObj.IsValid())
+				{
+					TArray<FString>	AssetList;
+					AssetList.Add(CueHandlerItem->GameplayCueNotifyObj.ToString());
+					FAssetEditorManager::Get().OpenEditorsForAssets(AssetList);
+				}
+				else if (CueHandlerItem->FunctionPtr.IsValid())
+				{
+					TArray<FString>	AssetList;
+					AssetList.Add(CueHandlerItem->FunctionPtr->GetOuter()->GetPathName());
+					FAssetEditorManager::Get().OpenEditorsForAssets(AssetList);
+				}
+			}
 
-	return FReply::Handled();
-}
+			TSharedPtr<FCueHandlerItem> CueHandlerItem;
+		};
 
-/** Builds content of the list in the GameplayCue Editor */
-void SGameplayCueEditor::UpdateGameplayCueListItems()
-{
-#if STATS
-	FString PerfMessage = FString::Printf(TEXT("SGameplayCueEditor::UpdateGameplayCueListItems"));
-	SCOPE_LOG_TIME_IN_SECONDS(*PerfMessage, nullptr)
-#endif
 
-	GameplayCueListItems.Reset();
-	IGameplayTagsModule& GameplayTagModule = IGameplayTagsModule::Get();
-	UGameplayCueManager* CueManager = UAbilitySystemGlobals::Get().GetGameplayCueManager();
-	TSharedPtr<FGameplayCueEditorItem> SelectItem;
 
-	FGameplayTagContainer AllGameplayCueTags = IGameplayTagsModule::Get().GetGameplayTagsManager().RequestGameplayTagChildren(UGameplayCueSet::BaseGameplayCueTag());
-	for (FGameplayTag ThisGameplayCueTag : AllGameplayCueTags)
-	{
-		if (SearchString.IsEmpty() == false && ThisGameplayCueTag.ToString().Contains(SearchString) == false)
-		{
-			continue;
-		}
 
-		bool Handled = false;
 
-		TSharedRef< FGameplayCueEditorItem > NewItem( new FGameplayCueEditorItem() );
-		NewItem->GameplayCueTagName = ThisGameplayCueTag;
 
-		if (AutoSelectTag.IsEmpty() == false && AutoSelectTag == ThisGameplayCueTag.ToString())
-		{
-			SelectItem = NewItem;
-		}
 
-		// Add notifies from the global set
-		if (CueManager && CueManager->GlobalCueSet)
-		{
-			int32* idxPtr = CueManager->GlobalCueSet->GameplayCueDataMap.Find(ThisGameplayCueTag);
-			if (idxPtr)
-			{
-				int32 idx = *idxPtr;
-				if (CueManager->GlobalCueSet->GameplayCueData.IsValidIndex(idx))
-				{
-					FGameplayCueNotifyData& Data = CueManager->GlobalCueSet->GameplayCueData[*idxPtr];
+		TSharedPtr<FCueItem> CueItem = InItem->AsCueItem();
+		TSharedPtr<FCueHandlerItem> CueHandlerItem = InItem->AsCueHandlerItem();
 
-					TSharedRef< FGameplayCueEditorHandlerItem > NewHandlerItem( new FGameplayCueEditorHandlerItem() );
-					NewHandlerItem->GameplayCueNotifyObj = Data.GameplayCueNotifyObj;
-					NewItem->HandlerItems.Add(NewHandlerItem);
-					Handled = true;
-				}
-			}
+		if ( CueItem.IsValid() )
+		{
+			return SNew(SCueItemWidget, GameplayCueTreeView.ToSharedRef(), FSimpleDelegate::CreateRaw(this, &SGameplayCueEditorImpl::UpdateGameplayCueListItems), CueItem, SearchText);
 		}
-
-		// Add events implemented by IGameplayCueInterface blueprints
-		TArray<UFunction*> Funcs;
-		EventMap.MultiFind(ThisGameplayCueTag, Funcs);
-
-		for (UFunction* Func : Funcs)
+		else if (CueHandlerItem.IsValid())
 		{
-			TSharedRef< FGameplayCueEditorHandlerItem > NewHandlerItem( new FGameplayCueEditorHandlerItem() );
-			NewHandlerItem->FunctionPtr = Func;
-			NewItem->HandlerItems.Add(NewHandlerItem);
-			Handled = true;
+			return SNew(SCueHandlerItemWidget, OwnerTable, CueHandlerItem);
 		}
-
-		if (bShowAll || Handled)
+		else
 		{
-			GameplayCueListItems.Add(NewItem);
+			return
+			SNew(STableRow< TSharedPtr<FTreeItem> >, OwnerTable)
+			[
+				SNew(STextBlock)
+				.Text(LOCTEXT("UnknownItemType", "Unknown Item Type"))
+			];
 		}
-
-		// Todo: Add notifies from all GameplayCueSets?
 	}
 
 
-	if (GameplayCueListView.IsValid())
+	void OnGetChildren(TSharedPtr<FTreeItem> Item, TArray< TSharedPtr<FTreeItem> >& Children)
 	{
-		GameplayCueListView->RequestListRefresh();
+		TSharedPtr<FCueItem> CueItem = Item->AsCueItem();
+		if (CueItem.IsValid())
+		{
+			Children.Append(CueItem->HandlerItems);
+		}
 	}
 
-	if (SelectItem.IsValid())
+	/** Builds content of the list in the GameplayCue Editor */
+	void UpdateGameplayCueListItems()
 	{
-		GameplayCueListView->SetItemSelection(SelectItem, true);
-		GameplayCueListView->RequestScrollIntoView(SelectItem);
-		
-	}
-	AutoSelectTag.Empty();
-}
-
-/** Widget for picking a new GameplayCue Notify class (similiar to actor class picker)  */
-class SGameplayCuePickerDialog : public SCompoundWidget
-{
-public:
-	SLATE_BEGIN_ARGS( SGameplayCuePickerDialog )
-	{}
-
-		SLATE_ARGUMENT( TSharedPtr<SWindow>, ParentWindow )
-		SLATE_ARGUMENT( TArray<UClass*>, DefaultClasses )
-		SLATE_ARGUMENT( FString, GameplayCueTag )
-
-	SLATE_END_ARGS()
-
-	/** Constructs this widget with InArgs */
-	void Construct(const FArguments& InArgs);
-
-	static bool PickGameplayCue(const FText& TitleText, const TArray<UClass*>& DefaultClasses, UClass*& OutChosenClass, FString InGameplayCueTag);
-
-private:
-	/** Handler for when a class is picked in the class picker */
-	void OnClassPicked(UClass* InChosenClass);
-
-	/** Creates the default class widgets */
-	TSharedRef<ITableRow> GenerateListRow(UClass* InItem, const TSharedRef<STableViewBase>& OwnerTable);
-
-	/** Handler for when "Ok" we selected in the class viewer */
-	FReply OnClassPickerConfirmed();
-
-	FReply OnDefaultClassPicked(UClass* InChosenClass);
-
-	/** Handler for the custom button to hide/unhide the default class viewer */
-	void OnDefaultAreaExpansionChanged(bool bExpanded);
-
-	/** Handler for the custom button to hide/unhide the class viewer */
-	void OnCustomAreaExpansionChanged(bool bExpanded);	
-
-private:
-	/** A pointer to the window that is asking the user to select a parent class */
-	TWeakPtr<SWindow> WeakParentWindow;
-
-	/** The class that was last clicked on */
-	UClass* ChosenClass;
-
-	/** A flag indicating that Ok was selected */
-	bool bPressedOk;
-
-	/**  An array of default classes used in the dialog **/
-	TArray<UClass*> DefaultClasses;
-
-	FString GameplayCueTag;
-};
-
-void SGameplayCuePickerDialog::Construct(const FArguments& InArgs)
-{
-	WeakParentWindow = InArgs._ParentWindow;
-	DefaultClasses = InArgs._DefaultClasses;
-	GameplayCueTag = InArgs._GameplayCueTag;
+#if STATS
+		FString PerfMessage = FString::Printf(TEXT("SGameplayCueEditor::UpdateGameplayCueListItems"));
+		SCOPE_LOG_TIME_IN_SECONDS(*PerfMessage, nullptr)
+#endif
 
-	FLinearColor AssetColor = FLinearColor::White;
+		GameplayCueListItems.Reset();
+		IGameplayTagsModule& GameplayTagModule = IGameplayTagsModule::Get();
+		UGameplayCueManager* CueManager = UAbilitySystemGlobals::Get().GetGameplayCueManager();
+		TSharedPtr<FCueItem> SelectItem;
 
-	bPressedOk = false;
-	ChosenClass = NULL;
+		FString SearchString = SearchText.ToString();
 
-	FString PathStr = SGameplayCueEditor::GetPathNameForGameplayCueTag(GameplayCueTag);
-	FGameplayCueEditorStrings Strings;
-	auto del = IGameplayAbilitiesEditorModule::Get().GetGameplayCueEditorStringsDelegate();
-	if (del.IsBound())
-	{
-		Strings = del.Execute();
-	}
+		FGameplayTagContainer AllGameplayCueTags = IGameplayTagsModule::Get().GetGameplayTagsManager().RequestGameplayTagChildren(UGameplayCueSet::BaseGameplayCueTag());
+		for (FGameplayTag ThisGameplayCueTag : AllGameplayCueTags)
+		{
+			if (SearchString.IsEmpty() == false && ThisGameplayCueTag.ToString().Contains(SearchString) == false)
+			{
+				continue;
+			}
 
-	ChildSlot
-	[
-		SNew(SBorder)
-		.Visibility(EVisibility::Visible)
-		.BorderImage(FEditorStyle::GetBrush("Menu.Background"))
-		[
-			SNew(SBox)
-			.Visibility(EVisibility::Visible)
-			.Padding(2.f)
-			.WidthOverride(520.0f)
-			[
-				SNew(SVerticalBox)
-				+SVerticalBox::Slot()
-				.Padding(2.f, 2.f)
-				.AutoHeight()
-				[
-					SNew(SBorder)
-					.Visibility(EVisibility::Visible)
-					.BorderImage( FEditorStyle::GetBrush("AssetThumbnail.AssetBackground") )
-					.BorderBackgroundColor(AssetColor.CopyWithNewOpacity(0.3f))
-					[
-						SNew(SExpandableArea)
-						.AreaTitle(NSLOCTEXT("SGameplayCuePickerDialog", "CommonClassesAreaTitle", "GameplayCue Notifies"))
-						.BodyContent()
-						[
-							SNew(SVerticalBox)
-							+SVerticalBox::Slot()
-							.Padding(2.f, 2.f)
-							.AutoHeight()
-							[
-								SNew(STextBlock)
-								
-								.Text(FText::FromString(Strings.GameplayCueNotifyDescription1))
-								.AutoWrapText(true)
-							]
+			bool Handled = false;
 
-							+SVerticalBox::Slot()
-							.AutoHeight( )
-							[
-								SNew(SListView < UClass*  >)
-								.ItemHeight(48)
-								.SelectionMode(ESelectionMode::None)
-								.ListItemsSource(&DefaultClasses)
-								.OnGenerateRow(this, &SGameplayCuePickerDialog::GenerateListRow)
-							]
+			TSharedRef< FCueItem > NewItem(new FCueItem());
+			NewItem->GameplayCueTagName = ThisGameplayCueTag;
 
-							+SVerticalBox::Slot()
-							.Padding(2.f, 2.f)
-							.AutoHeight()
-							[
-								SNew(STextBlock)
-								.Text(FText::FromString(FString::Printf(TEXT("This will create a new GameplayCue Notify here:"))))
-								.AutoWrapText(true)
-							]
-							+SVerticalBox::Slot()
-							.Padding(2.f, 2.f)
-							.AutoHeight()
-							[
-								SNew(STextBlock)
-								.Text(FText::FromString(PathStr))
-								.HighlightText(FText::FromString(PathStr))
-								.AutoWrapText(true)
-							]
-							+SVerticalBox::Slot()
-							.Padding(2.f, 2.f)
-							.AutoHeight()
-							[
-								SNew(STextBlock)
-								.Text(FText::FromString(Strings.GameplayCueNotifyDescription2))
-								.AutoWrapText(true)
-							]
-						]
-					]
-				]
+			if (AutoSelectTag.IsEmpty() == false && AutoSelectTag == ThisGameplayCueTag.ToString())
+			{
+				SelectItem = NewItem;
+			}
 
-				+SVerticalBox::Slot()
-				.Padding(2.f, 2.f)
-				.AutoHeight()
-				[
-					SNew(SBorder)
-					.Visibility(EVisibility::Visible)
-					.BorderImage( FEditorStyle::GetBrush("AssetThumbnail.AssetBackground") )
-					.BorderBackgroundColor(AssetColor.CopyWithNewOpacity(0.3f))
-					[
-						SNew(SExpandableArea)
-						.AreaTitle(NSLOCTEXT("SGameplayCuePickerDialogEvents", "CommonClassesAreaTitleEvents", "Custom BP Events"))
-						.BodyContent()
-						[
-							SNew(SVerticalBox)
-							+SVerticalBox::Slot()
-							.Padding(2.f, 2.f)
-							.AutoHeight()
-							[
-								SNew(STextBlock)
-								.Text(FText::FromString(Strings.GameplayCueEventDescription1))
-								.AutoWrapText(true)
+			// Add notifies from the global set
+			if (CueManager && CueManager->GlobalCueSet)
+			{
+				int32* idxPtr = CueManager->GlobalCueSet->GameplayCueDataMap.Find(ThisGameplayCueTag);
+				if (idxPtr)
+				{
+					int32 idx = *idxPtr;
+					if (CueManager->GlobalCueSet->GameplayCueData.IsValidIndex(idx))
+					{
+						FGameplayCueNotifyData& Data = CueManager->GlobalCueSet->GameplayCueData[*idxPtr];
 
-							]
+						TSharedRef< FCueHandlerItem > NewHandlerItem(new FCueHandlerItem());
+						NewHandlerItem->GameplayCueNotifyObj = Data.GameplayCueNotifyObj;
+						NewItem->HandlerItems.Add(NewHandlerItem);
+						Handled = true;
+					}
+				}
+			}
 
-							+SVerticalBox::Slot()
-							.Padding(2.f, 2.f)
-							.AutoHeight()
-							[
-								SNew(STextBlock)
-								.Text(FText::FromString(Strings.GameplayCueEventDescription2))
-								.AutoWrapText(true)
-							]							
-						]
-					]
-				]
-			]
-		]
-	];
-}
+			// Add events implemented by IGameplayCueInterface blueprints
+			TArray<UFunction*> Funcs;
+			EventMap.MultiFind(ThisGameplayCueTag, Funcs);
 
-/** Spawns window for picking new GameplayCue handler/notify */
-bool SGameplayCuePickerDialog::PickGameplayCue(const FText& TitleText, const TArray<UClass*>& DefaultClasses, UClass*& OutChosenClass, FString InGameplayCueName)
-{
-	// Create the window to pick the class
-	TSharedRef<SWindow> PickerWindow = SNew(SWindow)
-		.Title(TitleText)
-		.SizingRule( ESizingRule::Autosized )
-		.ClientSize( FVector2D( 0.f, 600.f ))
-		.SupportsMaximize(false)
-		.SupportsMinimize(false);
+			for (UFunction* Func : Funcs)
+			{
+				TSharedRef< FCueHandlerItem > NewHandlerItem(new FCueHandlerItem());
+				NewHandlerItem->FunctionPtr = Func;
+				NewItem->HandlerItems.Add(NewHandlerItem);
+				Handled = true;
+			}
 
-	TSharedRef<SGameplayCuePickerDialog> ClassPickerDialog = SNew(SGameplayCuePickerDialog)
-		.ParentWindow(PickerWindow)
-		.DefaultClasses(DefaultClasses)
-		.GameplayCueTag(InGameplayCueName);
+			if (bShowAll || Handled)
+			{
+				GameplayCueListItems.Add(NewItem);
+			}
 
-	PickerWindow->SetContent(ClassPickerDialog);
+			// If there are any handlers, automatically expand the item
+			if (NewItem->HandlerItems.Num() > 0)
+			{
+				GameplayCueTreeView->SetItemExpansion(NewItem, true);
+			}
 
-	GEditor->EditorAddModalWindow(PickerWindow);
+			// Todo: Add notifies from all GameplayCueSets?
+		}
 
-	if (ClassPickerDialog->bPressedOk)
-	{
-		OutChosenClass = ClassPickerDialog->ChosenClass;
-		return true;
+
+		if (GameplayCueTreeView.IsValid())
+		{
+			GameplayCueTreeView->RequestTreeRefresh();
+		}
+
+		if (SelectItem.IsValid())
+		{
+			GameplayCueTreeView->SetItemSelection(SelectItem, true);
+			GameplayCueTreeView->RequestScrollIntoView(SelectItem);
+
+		}
+		AutoSelectTag.Empty();
 	}
-	else
+
+	/** Slow task: use asset registry to find blueprints, load an inspect them to see what GC events they implement. */
+	FReply BuildEventMap()
 	{
-		// Ok was not selected, NULL the class
-		OutChosenClass = NULL;
-		return false;
+		FScopedSlowTask SlowTask(100.f, LOCTEXT("BuildEventMap", "Searching Blueprints for GameplayCue events"));
+		SlowTask.MakeDialog();
+		SlowTask.EnterProgressFrame(10.f);
+
+		EventMap.Empty();
+
+		IGameplayTagsModule& GameplayTagModule = IGameplayTagsModule::Get();
+
+		auto del = IGameplayAbilitiesEditorModule::Get().GetGameplayCueInterfaceClassesDelegate();
+		if (del.IsBound())
+		{
+			TArray<UClass*> InterfaceClasses;
+			del.ExecuteIfBound(InterfaceClasses);
+			float WorkPerClass = InterfaceClasses.Num() > 0 ? 90.f / static_cast<float>(InterfaceClasses.Num()) : 0.f;
+
+			for (UClass* InterfaceClass : InterfaceClasses)
+			{
+				SlowTask.EnterProgressFrame(WorkPerClass);
+
+				TArray<FAssetData> GameplayCueInterfaceActors;
+				{
+#if STATS
+					FString PerfMessage = FString::Printf(TEXT("Searched asset registry %s "), *InterfaceClass->GetName());
+					SCOPE_LOG_TIME_IN_SECONDS(*PerfMessage, nullptr)
+#endif
+
+						UObjectLibrary* ObjLibrary = UObjectLibrary::CreateLibrary(InterfaceClass, true, true);
+					ObjLibrary->LoadBlueprintAssetDataFromPath(TEXT("/Game/"));
+					ObjLibrary->GetAssetDataList(GameplayCueInterfaceActors);
+				}
+
+			{
+#if STATS
+				FString PerfMessage = FString::Printf(TEXT("Fully Loaded GameplayCueNotify actors %s "), *InterfaceClass->GetName());
+				SCOPE_LOG_TIME_IN_SECONDS(*PerfMessage, nullptr)
+#endif				
+
+					for (const FAssetData& AssetData : GameplayCueInterfaceActors)
+					{
+						UBlueprint* BP = Cast<UBlueprint>(AssetData.GetAsset());
+						if (BP)
+						{
+							for (TFieldIterator<UFunction> FuncIt(BP->GeneratedClass, EFieldIteratorFlags::ExcludeSuper); FuncIt; ++FuncIt)
+							{
+								FString FuncName = FuncIt->GetName();
+								if (FuncName.Contains("GameplayCue"))
+								{
+									FuncName.ReplaceInline(TEXT("_"), TEXT("."));
+									FGameplayTag FoundTag = GameplayTagModule.RequestGameplayTag(FName(*FuncName), false);
+									if (FoundTag.IsValid())
+									{
+										EventMap.AddUnique(FoundTag, *FuncIt);
+									}
+								}
+
+							}
+						}
+					}
+				}
+			}
+
+			UpdateGameplayCueListItems();
+		}
+
+		return FReply::Handled();
 	}
-}
 
-void SGameplayCuePickerDialog::OnClassPicked(UClass* InChosenClass)
-{
-	ChosenClass = InChosenClass;
-}
+	bool bShowAll;
 
-/** Generates rows in the list of GameplayCueNotify classes to pick from */
-TSharedRef<ITableRow> SGameplayCuePickerDialog::GenerateListRow(UClass* ItemClass, const TSharedRef<STableViewBase>& OwnerTable)
-{	
-	const FSlateBrush* ItemBrush = FClassIconFinder::FindIconForClass(ItemClass);
+	TSharedPtr< SGameplayCueTreeView > GameplayCueTreeView;
+	TArray< TSharedPtr< FTreeItem > > GameplayCueListItems;
 
-	return 
-	SNew(STableRow< UClass* >, OwnerTable)
+	TMultiMap<FGameplayTag, UFunction*> EventMap;
+
+	FString AutoSelectTag;
+	FText SearchText;
+};
+
+BEGIN_SLATE_FUNCTION_BUILD_OPTIMIZATION
+void SGameplayCueEditorImpl::Construct(const FArguments& InArgs)
+{
+	bShowAll = true;
+	bool CanAddFromINI = UGameplayTagsManager::ShouldImportTagsFromINI(); // We only support adding new tags to the ini files.
+	ChildSlot
 	[
 		SNew(SVerticalBox)
-		+SVerticalBox::Slot()
-		.MaxHeight(60.0f)
-		.Padding(10.0f, 6.0f, 0.0f, 4.0f)
+
+		+ SVerticalBox::Slot()
+		.AutoHeight()
 		[
 			SNew(SHorizontalBox)
-			+SHorizontalBox::Slot()
-			.FillWidth(0.45f)
+			+ SHorizontalBox::Slot()
+			.Padding(2.0f, 2.0f)
+			.AutoWidth()
 			[
+
 				SNew(SButton)
-				.OnClicked(this, &SGameplayCuePickerDialog::OnDefaultClassPicked, ItemClass)
-				.Content()
-				[
-					SNew(SHorizontalBox)
-					+SHorizontalBox::Slot()
-					.HAlign(HAlign_Center)
-					.VAlign(VAlign_Center)
-					.FillWidth(0.12f)
-					[
-						SNew(SImage)
-						.Image(ItemBrush)
-					]
-					+SHorizontalBox::Slot()
-					.VAlign(VAlign_Center)
-					.Padding(4.0f, 0.0f)
-					.FillWidth(0.8f)
-					[
-						SNew(STextBlock)
-						.Text(ItemClass->GetDisplayNameText())
-					]
+				.Text(LOCTEXT("SearchBPEvents", "Search BP Events"))
+				.OnClicked(this, &SGameplayCueEditorImpl::BuildEventMap)
+			]
 
+			+ SHorizontalBox::Slot()
+			.Padding(2.0f, 2.0f)
+			.AutoWidth()
+			[
+				// Check box that controls LIVE MODE
+				SNew(SCheckBox)
+				.IsChecked(this, &SGameplayCueEditorImpl::HandleShowAllCheckBoxIsChecked)
+				.OnCheckStateChanged(this, &SGameplayCueEditorImpl::HandleShowAllCheckedStateChanged)
+				[
+					SNew(STextBlock)
+					.Text(LOCTEXT("HideUnhandled", "Hide Unhandled GameplayCues"))
 				]
 			]
-			+SHorizontalBox::Slot()
-			.Padding(10.0f, 0.0f)
+		]
+
+		+ SVerticalBox::Slot()
+		.AutoHeight()
+		[
+			SNew(SHorizontalBox)
+			+ SHorizontalBox::Slot()
+			.Padding(2.0f, 2.0f)
+			.AutoWidth()
 			[
-				SNew(STextBlock)
-				.Text(ItemClass->GetToolTipText(true))
-				.AutoWrapText(true)
+				SAssignNew(NewGameplayCueTextBox, SEditableTextBox)
+				.MinDesiredWidth(210.0f)
+				.HintText(LOCTEXT("GameplayCueXY", "GameplayCue.X.Y"))
+				.OnTextCommitted(this, &SGameplayCueEditorImpl::OnNewGameplayCueTagCommited)
 			]
-			+SHorizontalBox::Slot()
+
+			+ SHorizontalBox::Slot()
+			.Padding(2.0f, 2.0f)
 			.AutoWidth()
 			[
-				FEditorClassUtils::GetDocumentationLinkWidget(ItemClass)
+				SNew(SButton)
+				.Text(LOCTEXT("AddNew", "Add New"))
+				.OnClicked(this, &SGameplayCueEditorImpl::OnNewGameplayCueButtonPressed)
+				.Visibility(CanAddFromINI ? EVisibility::Visible : EVisibility::Collapsed)
 			]
 		]
-	];
-}
 
-FReply SGameplayCuePickerDialog::OnDefaultClassPicked(UClass* InChosenClass)
-{
-	ChosenClass = InChosenClass;
-	bPressedOk = true;
-	if (WeakParentWindow.IsValid())
-	{
-		WeakParentWindow.Pin()->RequestDestroyWindow();
-	}
-	return FReply::Handled();
-}
+		+ SVerticalBox::Slot()
+		.AutoHeight()
+		[
+			SNew(SHorizontalBox)
+			+ SHorizontalBox::Slot()
+			.Padding(2.0f)
+			.AutoWidth()
+			[
+				SNew(SSearchBox)
+				.MinDesiredWidth(210.0f)
+				.OnTextCommitted(this, &SGameplayCueEditorImpl::OnSearchTagCommited)
+			]
+			+ SHorizontalBox::Slot()
+			.Padding(2.0f)
+			.AutoWidth()
+			[
+				SNew(SButton)
+				.Text(LOCTEXT("Search", "Search"))
+				.OnClicked(this, &SGameplayCueEditorImpl::DoSearch)
+			]
+		]
 
-FReply SGameplayCuePickerDialog::OnClassPickerConfirmed()
-{
-	if (ChosenClass == NULL)
-	{
-		FMessageDialog::Open(EAppMsgType::Ok, NSLOCTEXT("EditorFactories", "MustChooseClassWarning", "You must choose a class."));
-	}
-	else
-	{
-		bPressedOk = true;
+		+ SVerticalBox::Slot()
+		.FillHeight(1.0f)
+		[
+			SNew(SBorder)
+			.BorderImage(FEditorStyle::GetBrush("ToolBar.Background"))
+			[
+				SAssignNew(GameplayCueTreeView, SGameplayCueTreeView)
+				.ItemHeight(24)
+				.TreeItemsSource(&GameplayCueListItems)
+				.OnGenerateRow(this, &SGameplayCueEditorImpl::OnGenerateWidgetForGameplayCueListView)
+				.OnGetChildren( this, &SGameplayCueEditorImpl::OnGetChildren )
+				.HeaderRow
+				(
+					SNew(SHeaderRow)
+					+ SHeaderRow::Column(CueTagColumnName)
+					.DefaultLabel(NSLOCTEXT("GameplayCueEditor", "GameplayCueTag", "Tag"))
+					.FillWidth(0.25)
+
+					+ SHeaderRow::Column(CueHandlerColumnName)
+					.DefaultLabel(NSLOCTEXT("GameplayCueEditor", "GameplayCueHandlers", "Handlers"))
+					//.FillWidth(1000)
+				)
+			]
+		]
+	];
 
-		if (WeakParentWindow.IsValid())
-		{
-			WeakParentWindow.Pin()->RequestDestroyWindow();
-		}
-	}
-	return FReply::Handled();
+	UpdateGameplayCueListItems();
 }
 
-FString SGameplayCueEditor::GetPathNameForGameplayCueTag(FString GameplayCueTagName)
-{
-	FString NewDefaultPathName = FString::Printf(TEXT("/GC_%s"), *GameplayCueTagName);
-	auto PathDel = IGameplayAbilitiesEditorModule::Get().GetGameplayCueNotifyPathDelegate();
-	if (PathDel.IsBound())
-	{
-		NewDefaultPathName = PathDel.Execute(GameplayCueTagName);
-	}
-	NewDefaultPathName.ReplaceInline(TEXT("."), TEXT("_"));
-	return NewDefaultPathName;
-}
+BEGIN_SLATE_FUNCTION_BUILD_OPTIMIZATION
 
-/** Create new GameplayCueNotify: brings up dialogue to pick class, then creates it via the content browser. */
-FReply SGameplayCueEditor::FGameplayCueEditorItem::OnAddNewClicked()
+TSharedRef<SGameplayCueEditor> SGameplayCueEditor::New()
 {
-	FAssetToolsModule& AssetToolsModule = FModuleManager::GetModuleChecked<FAssetToolsModule>("AssetTools");
-	FContentBrowserModule& ContentBrowserModule = FModuleManager::LoadModuleChecked<FContentBrowserModule>("ContentBrowser");
-
-	TArray<UClass*>	NotifyClasses;
-	auto del = IGameplayAbilitiesEditorModule::Get().GetGameplayCueNotifyClassesDelegate();
-	del.ExecuteIfBound(NotifyClasses);
-	if (NotifyClasses.Num() == 0)
-	{
-		NotifyClasses.Add(UGameplayCueNotify_Static::StaticClass());
-		NotifyClasses.Add(AGameplayCueNotify_Actor::StaticClass());
-	}
-
-	// --------------------------------------------------
-
-	// Null the parent class to ensure one is selected	
-
-	const FText TitleText = LOCTEXT("CreateBlueprintOptions", "New GameplayCue Handler");
-	UClass* ChosenClass = nullptr;
-	const bool bPressedOk = SGameplayCuePickerDialog::PickGameplayCue(TitleText, NotifyClasses, ChosenClass, GameplayCueTagName.ToString());
-	if (bPressedOk)
-	{
-		FString NewDefaultPathName = SGameplayCueEditor::GetPathNameForGameplayCueTag(GameplayCueTagName.ToString());
-
-		// Make sure the name is unique
-		FString AssetName;
-		FString PackageName;
-		AssetToolsModule.Get().CreateUniqueAssetName(NewDefaultPathName , TEXT(""), /*out*/ PackageName, /*out*/ AssetName);
-		const FString PackagePath = FPackageName::GetLongPackagePath(PackageName);
-
-		// Create the GameplayCue Notify
-		UBlueprintFactory* BlueprintFactory = NewObject<UBlueprintFactory>();
-		BlueprintFactory->ParentClass = ChosenClass;
-		ContentBrowserModule.Get().CreateNewAsset(AssetName, PackagePath, UBlueprint::StaticClass(), BlueprintFactory);
-	}
-	
-	return FReply::Handled();
+	return MakeShareable(new SGameplayCueEditorImpl());
 }
 
 #undef LOCTEXT_NAMESPACE
diff --git a/Engine/Source/Editor/GameplayAbilitiesEditor/Private/SGameplayCueEditor.h b/Engine/Source/Editor/GameplayAbilitiesEditor/Private/SGameplayCueEditor.h
--- a/Engine/Source/Editor/GameplayAbilitiesEditor/Private/SGameplayCueEditor.h
+++ b/Engine/Source/Editor/GameplayAbilitiesEditor/Private/SGameplayCueEditor.h
@@ -6,76 +6,18 @@
 #include "GameplayTagContainer.h"
 
 /** Main CollisionAnalyzer UI widget */
-class SGameplayCueEditor : public SCompoundWidget
+class SGameplayCueEditor : public SUserWidget
 {
 public:
 
-	SLATE_BEGIN_ARGS(SGameplayCueEditor) {}
-
+	SLATE_USER_ARGS(SGameplayCueEditor) {}
 	SLATE_END_ARGS()
+	virtual void Construct(const FArguments& InArgs) = 0;
 
-	void Construct(const FArguments& InArgs);
-
-	virtual ~SGameplayCueEditor();
-
-	void OnNewGameplayCueTagCommited(const FText& InText, ETextCommit::Type InCommitType);
-	void OnSearchTagCommited(const FText& InText, ETextCommit::Type InCommitType);
+	virtual void OnNewGameplayCueTagCommited(const FText& InText, ETextCommit::Type InCommitType) = 0;
+	virtual void OnSearchTagCommited(const FText& InText, ETextCommit::Type InCommitType) = 0;
 
-	FReply OnNewGameplayCueButtonPressed();
+	virtual FReply OnNewGameplayCueButtonPressed() = 0;
 	
 	static FString GetPathNameForGameplayCueTag(FString Tag);
-
-private:
-
-	TSharedPtr<SEditableTextBox>	NewGameplayCueTextBox;
-
-	void CreateNewGameplayCueTag();
-
-	void HandleShowAllCheckedStateChanged( ECheckBoxState NewValue );
-
-	/** Callback for getting the checked state of the focus check box. */
-	ECheckBoxState HandleShowAllCheckBoxIsChecked() const;
-
-	bool bShowAll;
-
-	struct FGameplayCueEditorHandlerItem
-	{
-		FStringAssetReference GameplayCueNotifyObj;
-
-		TWeakObjectPtr<UFunction> FunctionPtr;
-	};
-
-	typedef SListView< TSharedPtr< FGameplayCueEditorHandlerItem > > SGameplayCueHandlerListView;
-	typedef TArray< TSharedPtr< FGameplayCueEditorHandlerItem > > FGameplayCueHandlerArray;
-
-	/** An item in the GameplayCue list */
-	struct FGameplayCueEditorItem
-	{
-		FGameplayTag GameplayCueTagName;
-
-		FGameplayCueHandlerArray HandlerItems;
-
-		TSharedPtr< SGameplayCueHandlerListView > GameplayCueHandlerListView;
-
-		FReply OnAddNewClicked();
-		void OnNewClassPicked(FString SelectedPath);
-	};
-
-	typedef SListView< TSharedPtr< FGameplayCueEditorItem > > SGameplayCueListView;
-	typedef TArray< TSharedPtr< FGameplayCueEditorItem > > FGameplayCueArray;
-	
-
-	TSharedRef<ITableRow> OnGenerateWidgetForGameplayCueListView(TSharedPtr< FGameplayCueEditorItem > InItem, const TSharedRef<STableViewBase>& OwnerTable);
-
-	TSharedPtr< SGameplayCueListView > GameplayCueListView;
-	FGameplayCueArray GameplayCueListItems;
-
-	void UpdateGameplayCueListItems();
-
-	FReply BuildEventMap();
-	
-	TMultiMap<FGameplayTag, UFunction*> EventMap;
-
-	FString AutoSelectTag;
-	FString SearchString;
 };
