update default gameplay cue notify path to be valid

--------
Integrated using branch UE4-To-UE4-Orion (reversed) of change#2590115 by [REDACTED_PERSON] on 2015/06/17 10:46:40.

diff --git a/Engine/Source/Editor/GameplayAbilitiesEditor/Private/SGameplayCueEditor.cpp b/Engine/Source/Editor/GameplayAbilitiesEditor/Private/SGameplayCueEditor.cpp
--- a/Engine/Source/Editor/GameplayAbilitiesEditor/Private/SGameplayCueEditor.cpp
+++ b/Engine/Source/Editor/GameplayAbilitiesEditor/Private/SGameplayCueEditor.cpp
@@ -341,7 +341,7 @@ FReply SGameplayCuePickerDialog::OnClassPickerConfirmed()
 
 FString SGameplayCueEditor::GetPathNameForGameplayCueTag(FString GameplayCueTagName)
 {
-	FString NewDefaultPathName = FString::Printf(TEXT("/GC_%s"), *GameplayCueTagName);
+	FString NewDefaultPathName = FString::Printf(TEXT("/Game/GC_%s"), *GameplayCueTagName);
 	auto PathDel = IGameplayAbilitiesEditorModule::Get().GetGameplayCueNotifyPathDelegate();
 	if (PathDel.IsBound())
 	{
