Merge missed changes to GameplayAbilitiesEditor.Build.cs

diff --git a/Engine/Source/Editor/GameplayAbilitiesEditor/GameplayAbilitiesEditor.Build.cs b/Engine/Source/Editor/GameplayAbilitiesEditor/GameplayAbilitiesEditor.Build.cs
--- a/Engine/Source/Editor/GameplayAbilitiesEditor/GameplayAbilitiesEditor.Build.cs
+++ b/Engine/Source/Editor/GameplayAbilitiesEditor/GameplayAbilitiesEditor.Build.cs
@@ -40,10 +40,12 @@ namespace UnrealBuildTool.Rules
 					"GraphEditor",
 					"MainFrame",
 					"UnrealEd",
+                    "WorkspaceMenuStructure",
+                    "ContentBrowser",
                     "EditorWidgets",
                     "SourceControl",
 				}
 			);
 		}
 	}
-}
\ No newline at end of file
+}
