Merging changes to gameplay tags and gameplay cue editor.

diff --git a/Engine/Source/Editor/GameplayAbilitiesEditor/Private/GameplayAbilitiesEditorModule.cpp b/Engine/Source/Editor/GameplayAbilitiesEditor/Private/GameplayAbilitiesEditorModule.cpp
--- a/Engine/Source/Editor/GameplayAbilitiesEditor/Private/GameplayAbilitiesEditorModule.cpp
+++ b/Engine/Source/Editor/GameplayAbilitiesEditor/Private/GameplayAbilitiesEditorModule.cpp
@@ -15,6 +15,7 @@
 #include "AssetTypeActions_GameplayEffect.h"
 #include "GameplayAbilitiesGraphPanelPinFactory.h"
 #include "GameplayAbilitiesGraphPanelNodeFactory.h"
+#include "GameplayCueTagDetails.h"
 
 #include "Runtime/GameplayTags/Public/GameplayTagsModule.h"
 #include "Editor/BlueprintGraph/Public/BlueprintActionDatabase.h"
@@ -92,6 +93,7 @@ void FGameplayAbilitiesEditorModule::StartupModule()
 	PropertyModule.RegisterCustomPropertyTypeLayout( "GameplayEffectExecutionScopedModifierInfo", FOnGetPropertyTypeCustomizationInstance::CreateStatic( &FGameplayEffectExecutionScopedModifierInfoDetails::MakeInstance ) );
 	PropertyModule.RegisterCustomPropertyTypeLayout( "GameplayEffectExecutionDefinition", FOnGetPropertyTypeCustomizationInstance::CreateStatic( &FGameplayEffectExecutionDefinitionDetails::MakeInstance ) );
 	PropertyModule.RegisterCustomPropertyTypeLayout( "GameplayEffectModifierMagnitude", FOnGetPropertyTypeCustomizationInstance::CreateStatic( &FGameplayEffectModifierMagnitudeDetails::MakeInstance ) );
+	PropertyModule.RegisterCustomPropertyTypeLayout( "GameplayCueTag", FOnGetPropertyTypeCustomizationInstance::CreateStatic( &FGameplayCueTagDetails::MakeInstance ) );
 
 	PropertyModule.RegisterCustomClassLayout( "AttributeSet", FOnGetDetailCustomizationInstance::CreateStatic( &FAttributeDetails::MakeInstance ) );
 	PropertyModule.RegisterCustomClassLayout( "GameplayEffect", FOnGetDetailCustomizationInstance::CreateStatic( &FGameplayEffectDetails::MakeInstance ) );
diff --git a/Engine/Source/Editor/GameplayAbilitiesEditor/Private/GameplayCueTagDetails.cpp b/Engine/Source/Editor/GameplayAbilitiesEditor/Private/GameplayCueTagDetails.cpp
new file mode 100644
--- /dev/null
+++ b/Engine/Source/Editor/GameplayAbilitiesEditor/Private/GameplayCueTagDetails.cpp
@@ -0,0 +1,220 @@
+// Copyright 1998-2015 Epic Games, Inc. All Rights Reserved.
+
+#include "AbilitySystemEditorPrivatePCH.h"
+#include "GameplayCueTagDetails.h"
+#include "Editor/PropertyEditor/Public/PropertyEditing.h"
+#include "KismetEditorUtilities.h"
+#include "AttributeSet.h"
+#include "GameplayAbilitiesModule.h"
+#include "SGameplayAttributeGraphPin.h"
+#include "SSearchBox.h"
+#include "STextComboBox.h"
+#include "GameplayEffect.h"
+#include "AbilitySystemComponent.h"
+#include "GameplayEffectExtension.h"
+#include "GameplayCueInterface.h"
+#include "SHyperlink.h"
+#include "GameplayTagsManager.h"
+#include "GameplayCueManager.h"
+#include "GameplayCueSet.h"
+#include "SGameplayCueEditor.h"
+
+#define LOCTEXT_NAMESPACE "GameplayCueDetailsCustomization"
+
+DEFINE_LOG_CATEGORY(LogGameplayCueDetails);
+
+TSharedRef<IPropertyTypeCustomization> FGameplayCueTagDetails::MakeInstance()
+{
+	return MakeShareable(new FGameplayCueTagDetails());
+}
+
+void FGameplayCueTagDetails::CustomizeHeader( TSharedRef<IPropertyHandle> StructPropertyHandle, class FDetailWidgetRow& HeaderRow, IPropertyTypeCustomizationUtils& StructCustomizationUtils )
+{
+	uint32 NumChildren = 0;
+	StructPropertyHandle->GetNumChildren(NumChildren);
+
+	GameplayTagProperty = StructPropertyHandle->GetChildHandle(GET_MEMBER_NAME_CHECKED(FGameplayCueTag,GameplayCueTag));
+
+	FSimpleDelegate OnTagChanged = FSimpleDelegate::CreateSP(this, &FGameplayCueTagDetails::OnPropertyValueChanged);
+	StructPropertyHandle->SetOnPropertyValueChanged(OnTagChanged);
+
+	UGameplayCueManager* CueManager = UAbilitySystemGlobals::Get().GetGameplayCueManager();
+	if (CueManager)
+	{
+		CueManager->OnGameplayCueNotifyAddOrRemove.AddSP(this, &FGameplayCueTagDetails::OnPropertyValueChanged);
+	}
+
+	HeaderRow
+	.NameContent()
+	[
+		StructPropertyHandle->CreatePropertyNameWidget()
+	];
+}
+
+FText FGameplayCueTagDetails::GetTagText() const
+{
+	FString Str;
+	FGameplayTag* TagPtr = GetTag();
+	if (TagPtr && TagPtr->IsValid())
+	{
+		Str = TagPtr->GetTagName().ToString();
+	}
+
+	return FText::FromString(Str);
+}
+
+void FGameplayCueTagDetails::CustomizeChildren( TSharedRef<IPropertyHandle> StructPropertyHandle, class IDetailChildrenBuilder& StructBuilder, IPropertyTypeCustomizationUtils& StructCustomizationUtils )
+{
+	GameplayTagProperty = StructPropertyHandle->GetChildHandle(GET_MEMBER_NAME_CHECKED(FGameplayCueTag,GameplayCueTag));
+	if (GameplayTagProperty.IsValid())
+	{
+		IDetailPropertyRow& PropRow = StructBuilder.AddChildProperty(GameplayTagProperty.ToSharedRef());
+	}
+
+	bool ValidTag = UpdateNotifyList();
+	bool HasNotify = (NotifyList.Num() > 0);
+
+	StructBuilder.AddChildContent( LOCTEXT("NotifyLinkStr", "Notify") )
+	.NameContent()
+	[
+		StructPropertyHandle->CreatePropertyNameWidget(LOCTEXT("NotifyStr", "Notify"))
+	]
+	.ValueContent()
+	.MaxDesiredWidth(512)
+	[
+		SNew(SVerticalBox)
+		+SVerticalBox::Slot()
+		.Padding(2.f, 2.f)
+		.AutoHeight()
+		[
+			SNew(SHorizontalBox)
+			+ SHorizontalBox::Slot()
+			.Padding(2.0f, 2.0f)
+			.AutoWidth()
+			[
+				SAssignNew(ListView, SListView < TSharedRef<FStringAssetReference> >)
+				.ItemHeight(48)
+				.SelectionMode(ESelectionMode::None)
+				.ListItemsSource(&NotifyList)
+				.OnGenerateRow(this, &FGameplayCueTagDetails::GenerateListRow)
+				.Visibility(this, &FGameplayCueTagDetails::GetListViewVisibility)
+			]
+		]
+
+		+SVerticalBox::Slot()
+		.AutoHeight( )
+		[
+			SNew(SButton)
+			.Text(LOCTEXT("AddNew", "Add New"))
+			.Visibility(this, &FGameplayCueTagDetails::GetAddNewNotifyVisibliy)
+			.OnClicked(this, &FGameplayCueTagDetails::OnAddNewNotifyClicked)
+		]
+	];
+}
+
+TSharedRef<ITableRow> FGameplayCueTagDetails::GenerateListRow(TSharedRef<FStringAssetReference> NotifyName, const TSharedRef<STableViewBase>& OwnerTable)
+{	
+	FString ShortName = NotifyName->ToString();
+	int32 idx;
+	if (ShortName.FindLastChar(TEXT('.'), idx))
+	{
+		ShortName = ShortName.Right(ShortName.Len() - (idx+1));
+		ShortName.RemoveFromEnd(TEXT("_c"));
+	}
+
+	return
+	SNew(STableRow< TSharedRef<FStringAssetReference> >, OwnerTable)
+	[
+		SNew(SBox)
+		.HAlign(HAlign_Left)
+		[
+			SNew(SHyperlink)
+			.Style(FEditorStyle::Get(), "Common.GotoBlueprintHyperlink")
+			.Text(FText::FromString(ShortName))
+			.OnNavigate(this, &FGameplayCueTagDetails::NavigateToHandler, NotifyName)
+		]
+	];
+}
+
+void FGameplayCueTagDetails::NavigateToHandler(TSharedRef<FStringAssetReference> AssetRef)
+{
+	SGameplayCueEditor::OpenEditorForNotify(AssetRef->ToString());
+}
+
+FReply FGameplayCueTagDetails::OnAddNewNotifyClicked()
+{
+	FGameplayTag* TagPtr = GetTag();
+	if (TagPtr && TagPtr->IsValid())
+	{
+		SGameplayCueEditor::CreateNewGameplayCueNotifyDialogue(TagPtr->ToString());
+		OnPropertyValueChanged();
+	}
+	return FReply::Handled();
+}
+
+void FGameplayCueTagDetails::OnPropertyValueChanged()
+{
+	UpdateNotifyList();
+	if (ListView.IsValid())
+	{
+		ListView->RequestListRefresh();
+	}
+}
+
+bool FGameplayCueTagDetails::UpdateNotifyList()
+{
+	NotifyList.Empty();
+	
+	FGameplayTag* Tag = GetTag();
+	bool ValidTag = Tag && Tag->IsValid();
+	if (ValidTag)
+	{
+		uint8 EnumVal;
+		GameplayTagProperty->GetValue(EnumVal);
+
+		UGameplayCueManager* CueManager = UAbilitySystemGlobals::Get().GetGameplayCueManager();
+		if (CueManager && CueManager->GlobalCueSet)
+		{
+			if (int32* idxPtr = CueManager->GlobalCueSet->GameplayCueDataMap.Find(*Tag))
+			{
+				int32 idx = *idxPtr;
+				if (CueManager->GlobalCueSet->GameplayCueData.IsValidIndex(idx))
+				{
+					FGameplayCueNotifyData& Data = CueManager->GlobalCueSet->GameplayCueData[*idxPtr];
+
+					TSharedRef<FStringAssetReference> Item(MakeShareable(new FStringAssetReference(Data.GameplayCueNotifyObj)));
+					NotifyList.Add(Item);
+				}
+			}
+		}
+	}
+
+	return ValidTag;
+}
+
+FGameplayTag* FGameplayCueTagDetails::GetTag() const
+{
+	FGameplayTag* Tag = nullptr;
+	TArray<void*> RawStructData;
+	if (GameplayTagProperty.IsValid())
+	{
+		GameplayTagProperty->AccessRawData(RawStructData);
+		Tag = (FGameplayTag*)(RawStructData[0]);
+	}
+	return Tag;
+}
+
+EVisibility FGameplayCueTagDetails::GetAddNewNotifyVisibliy() const
+{
+	FGameplayTag* Tag = GetTag();
+	bool ValidTag = Tag && Tag->IsValid();
+	bool HasNotify = NotifyList.Num() > 0;
+	return (!HasNotify && ValidTag) ? EVisibility::Visible : EVisibility::Collapsed;
+}
+
+EVisibility FGameplayCueTagDetails::GetListViewVisibility() const
+{
+	return NotifyList.Num() > 0 ? EVisibility::Visible : EVisibility::Collapsed;
+}
+
+#undef LOCTEXT_NAMESPACE
diff --git a/Engine/Source/Editor/GameplayAbilitiesEditor/Private/GameplayCueTagDetails.h b/Engine/Source/Editor/GameplayAbilitiesEditor/Private/GameplayCueTagDetails.h
new file mode 100644
--- /dev/null
+++ b/Engine/Source/Editor/GameplayAbilitiesEditor/Private/GameplayCueTagDetails.h
@@ -0,0 +1,48 @@
+// Copyright 1998-2015 Epic Games, Inc. All Rights Reserved.
+
+#pragma once
+
+#include "Editor/PropertyEditor/Public/PropertyEditorModule.h"
+
+class IPropertyHandle;
+class FDetailWidgetRow;
+class IDetailLayoutBuilder;
+
+struct FGameplayTag;
+
+DECLARE_LOG_CATEGORY_EXTERN(LogGameplayCueDetails, Log, All);
+
+class FGameplayCueTagDetails : public IPropertyTypeCustomization
+{
+public:
+	static TSharedRef<IPropertyTypeCustomization> MakeInstance();
+
+
+private:
+
+	FGameplayCueTagDetails()
+	{
+	}
+
+	virtual void CustomizeHeader( TSharedRef<class IPropertyHandle> StructPropertyHandle, class FDetailWidgetRow& HeaderRow, IPropertyTypeCustomizationUtils& StructCustomizationUtils ) override;
+	virtual void CustomizeChildren( TSharedRef<class IPropertyHandle> StructPropertyHandle, class IDetailChildrenBuilder& StructBuilder, IPropertyTypeCustomizationUtils& StructCustomizationUtils ) override;	
+
+	TSharedPtr<IPropertyHandle> GameplayTagProperty;
+
+	TSharedRef<ITableRow> GenerateListRow(TSharedRef<FStringAssetReference> NotifyName, const TSharedRef<STableViewBase>& OwnerTable);
+	TArray<TSharedRef<FStringAssetReference> > NotifyList;
+	TSharedPtr< SListView < TSharedRef<FStringAssetReference> > > ListView;
+
+	void NavigateToHandler(TSharedRef<FStringAssetReference> AssetRef);
+
+	FReply OnAddNewNotifyClicked();
+
+	/** Updates the selected tag*/
+	void OnPropertyValueChanged();
+	bool UpdateNotifyList();
+	FGameplayTag* GetTag() const;
+	FText GetTagText() const;
+
+	EVisibility GetListViewVisibility() const;
+	EVisibility GetAddNewNotifyVisibliy() const;
+};
\ No newline at end of file
diff --git a/Engine/Source/Editor/GameplayAbilitiesEditor/Private/SGameplayCueEditor.cpp b/Engine/Source/Editor/GameplayAbilitiesEditor/Private/SGameplayCueEditor.cpp
--- a/Engine/Source/Editor/GameplayAbilitiesEditor/Private/SGameplayCueEditor.cpp
+++ b/Engine/Source/Editor/GameplayAbilitiesEditor/Private/SGameplayCueEditor.cpp
@@ -211,7 +211,6 @@ void SGameplayCuePickerDialog::Construct(const FArguments& InArgs)
 }
 END_SLATE_FUNCTION_BUILD_OPTIMIZATION
 
-
 static const FName CueTagColumnName("GameplayCueTags");
 static const FName CueHandlerColumnName("GameplayCueHandlers");
 
@@ -406,59 +405,17 @@ private:
 		FScopedSlowTask SlowTask(0.f, LOCTEXT("AddingNewGameplaycue", "Adding new GameplayCue Tag"));
 		SlowTask.MakeDialog();
 
-
 		FString str = NewGameplayCueTextBox->GetText().ToString();
-
-		AutoSelectTag = str;
-
-		UGameplayTagsSettings* Settings = GetMutableDefault<UGameplayTagsSettings>();
-
-		if (Settings)
+		if (str.IsEmpty())
 		{
+			return;
+		}
 
-			FString RelativeConfigFilePath = Settings->GetDefaultConfigFilename();
-			FString ConfigPath = FPaths::ConvertRelativePathToFull(RelativeConfigFilePath);
-
-			if (ISourceControlModule::Get().IsEnabled())
-			{
-				FText ErrorMessage;
-
-				if (!SourceControlHelpers::CheckoutOrMarkForAdd(ConfigPath, FText::FromString(ConfigPath), NULL, ErrorMessage))
-				{
-					FNotificationInfo Info(ErrorMessage);
-					Info.ExpireDuration = 3.0f;
-					FSlateNotificationManager::Get().AddNotification(Info);
-				}
-			}
-			else
-			{
-				if (!FPlatformFileManager::Get().GetPlatformFile().SetReadOnly(*ConfigPath, false))
-				{
-					FText NotificationErrorText = FText::Format(LOCTEXT("FailedToMakeWritable", "Could not make {0} writable."), FText::FromString(ConfigPath));
-
-					FNotificationInfo Info(NotificationErrorText);
-					Info.ExpireDuration = 3.0f;
-
-					FSlateNotificationManager::Get().AddNotification(Info);
-				}
-			}
-
-			Settings->GameplayTags.Add(str);
-			Settings->SortTags();
-			IGameplayTagsModule::Get().GetGameplayTagsManager().DestroyGameplayTagTree();
-
-			{
-#if STATS
-				FString PerfMessage = FString::Printf(TEXT("ConstructGameplayTagTree GameplayTag tables"));
-				SCOPE_LOG_TIME_IN_SECONDS(*PerfMessage, nullptr)
-#endif
-					IGameplayTagsModule::Get().GetGameplayTagsManager().ConstructGameplayTagTree();
-			}
+		AutoSelectTag = str;
 
-			Settings->UpdateDefaultConfigFile();
+		UGameplayTagsManager::AddNewGameplayTagToINI(str);
 
-			UpdateGameplayCueListItems();
-		}
+		UpdateGameplayCueListItems();
 
 		NewGameplayCueTextBox->SetText(FText::GetEmpty());
 	}
@@ -581,43 +538,8 @@ private:
 			/** Create new GameplayCueNotify: brings up dialogue to pick class, then creates it via the content browser. */
 			FReply OnAddNewClicked()
 			{
-				FAssetToolsModule& AssetToolsModule = FModuleManager::GetModuleChecked<FAssetToolsModule>("AssetTools");
-				FContentBrowserModule& ContentBrowserModule = FModuleManager::LoadModuleChecked<FContentBrowserModule>("ContentBrowser");
-
-				TArray<UClass*>	NotifyClasses;
-				auto del = IGameplayAbilitiesEditorModule::Get().GetGameplayCueNotifyClassesDelegate();
-				del.ExecuteIfBound(NotifyClasses);
-				if (NotifyClasses.Num() == 0)
-				{
-					NotifyClasses.Add(UGameplayCueNotify_Static::StaticClass());
-					NotifyClasses.Add(AGameplayCueNotify_Actor::StaticClass());
-				}
-
-				// --------------------------------------------------
-
-				// Null the parent class to ensure one is selected	
-
-				const FText TitleText = LOCTEXT("CreateBlueprintOptions", "New GameplayCue Handler");
-				UClass* ChosenClass = nullptr;
-				const bool bPressedOk = SGameplayCuePickerDialog::PickGameplayCue(TitleText, NotifyClasses, ChosenClass, Item->GameplayCueTagName.ToString());
-				if (bPressedOk)
-				{
-					FString NewDefaultPathName = SGameplayCueEditor::GetPathNameForGameplayCueTag(Item->GameplayCueTagName.ToString());
-
-					// Make sure the name is unique
-					FString AssetName;
-					FString PackageName;
-					AssetToolsModule.Get().CreateUniqueAssetName(NewDefaultPathName, TEXT(""), /*out*/ PackageName, /*out*/ AssetName);
-					const FString PackagePath = FPackageName::GetLongPackagePath(PackageName);
-
-					// Create the GameplayCue Notify
-					UBlueprintFactory* BlueprintFactory = NewObject<UBlueprintFactory>();
-					BlueprintFactory->ParentClass = ChosenClass;
-					ContentBrowserModule.Get().CreateNewAsset(AssetName, PackagePath, UBlueprint::StaticClass(), BlueprintFactory);
-
-					RefreshDelegate.ExecuteIfBound();
-				}
-
+				SGameplayCueEditor::CreateNewGameplayCueNotifyDialogue(Item->GameplayCueTagName.ToString());
+				RefreshDelegate.ExecuteIfBound();
 				return FReply::Handled();
 			}
 
@@ -706,27 +628,18 @@ private:
 			{
 				if (CueHandlerItem->GameplayCueNotifyObj.IsValid())
 				{
-					TArray<FString>	AssetList;
-					AssetList.Add(CueHandlerItem->GameplayCueNotifyObj.ToString());
-					FAssetEditorManager::Get().OpenEditorsForAssets(AssetList);
+					SGameplayCueEditor::OpenEditorForNotify(CueHandlerItem->GameplayCueNotifyObj.ToString());
+					
 				}
 				else if (CueHandlerItem->FunctionPtr.IsValid())
 				{
-					TArray<FString>	AssetList;
-					AssetList.Add(CueHandlerItem->FunctionPtr->GetOuter()->GetPathName());
-					FAssetEditorManager::Get().OpenEditorsForAssets(AssetList);
+					SGameplayCueEditor::OpenEditorForNotify(CueHandlerItem->FunctionPtr->GetOuter()->GetPathName());
 				}
 			}
 
 			TSharedPtr<FCueHandlerItem> CueHandlerItem;
 		};
 
-
-
-
-
-
-
 		TSharedPtr<FCueItem> CueItem = InItem->AsCueItem();
 		TSharedPtr<FCueHandlerItem> CueHandlerItem = InItem->AsCueHandlerItem();
 
@@ -749,6 +662,11 @@ private:
 		}
 	}
 
+	void OnPropertyValueChanged()
+	{
+		UpdateGameplayCueListItems();
+	}
+
 
 	void OnGetChildren(TSharedPtr<FTreeItem> Item, TArray< TSharedPtr<FTreeItem> >& Children)
 	{
@@ -936,6 +854,12 @@ private:
 BEGIN_SLATE_FUNCTION_BUILD_OPTIMIZATION
 void SGameplayCueEditorImpl::Construct(const FArguments& InArgs)
 {
+	UGameplayCueManager* CueManager = UAbilitySystemGlobals::Get().GetGameplayCueManager();
+	if (CueManager)
+	{
+		CueManager->OnGameplayCueNotifyAddOrRemove.AddSP(this, &SGameplayCueEditorImpl::OnPropertyValueChanged);
+	}
+
 	bShowAll = true;
 	bool CanAddFromINI = UGameplayTagsManager::ShouldImportTagsFromINI(); // We only support adding new tags to the ini files.
 	ChildSlot
@@ -1047,11 +971,74 @@ void SGameplayCueEditorImpl::Construct(const FArguments& InArgs)
 	UpdateGameplayCueListItems();
 }
 
-BEGIN_SLATE_FUNCTION_BUILD_OPTIMIZATION
+END_SLATE_FUNCTION_BUILD_OPTIMIZATION
 
 TSharedRef<SGameplayCueEditor> SGameplayCueEditor::New()
 {
 	return MakeShareable(new SGameplayCueEditorImpl());
 }
 
+void SGameplayCueEditor::CreateNewGameplayCueNotifyDialogue(FString GameplayCue)
+{
+	FAssetToolsModule& AssetToolsModule = FModuleManager::GetModuleChecked<FAssetToolsModule>("AssetTools");
+	FContentBrowserModule& ContentBrowserModule = FModuleManager::LoadModuleChecked<FContentBrowserModule>("ContentBrowser");
+
+	TArray<UClass*>	NotifyClasses;
+	auto del = IGameplayAbilitiesEditorModule::Get().GetGameplayCueNotifyClassesDelegate();
+	del.ExecuteIfBound(NotifyClasses);
+	if (NotifyClasses.Num() == 0)
+	{
+		NotifyClasses.Add(UGameplayCueNotify_Static::StaticClass());
+		NotifyClasses.Add(AGameplayCueNotify_Actor::StaticClass());
+	}
+
+	// --------------------------------------------------
+
+	// Null the parent class to ensure one is selected	
+
+	const FText TitleText = LOCTEXT("CreateBlueprintOptions", "New GameplayCue Handler");
+	UClass* ChosenClass = nullptr;
+	const bool bPressedOk = SGameplayCuePickerDialog::PickGameplayCue(TitleText, NotifyClasses, ChosenClass, GameplayCue);
+	if (bPressedOk)
+	{
+		FString NewDefaultPathName = SGameplayCueEditor::GetPathNameForGameplayCueTag(GameplayCue);
+
+		// Make sure the name is unique
+		FString AssetName;
+		FString PackageName;
+		AssetToolsModule.Get().CreateUniqueAssetName(NewDefaultPathName, TEXT(""), /*out*/ PackageName, /*out*/ AssetName);
+		const FString PackagePath = FPackageName::GetLongPackagePath(PackageName);
+
+		// Create the GameplayCue Notify
+		UBlueprintFactory* BlueprintFactory = NewObject<UBlueprintFactory>();
+		BlueprintFactory->ParentClass = ChosenClass;
+		ContentBrowserModule.Get().CreateNewAsset(AssetName, PackagePath, UBlueprint::StaticClass(), BlueprintFactory);
+	}
+}
+
+void SGameplayCueEditor::OpenEditorForNotify(FString NotifyFullPath)
+{
+	// This nonsense is to handle the case where the asset only exists in memory
+	// and there for does not have a linker/exist on disk. (The FString version
+	// of OpenEditorForAsset does not handle this).
+	FStringAssetReference AssetRef(NotifyFullPath);
+
+	UObject* Obj = AssetRef.ResolveObject();
+	if (!Obj)
+	{
+		Obj = AssetRef.TryLoad();
+	}
+
+	if (Obj)
+	{
+		UPackage* pkg = Cast<UPackage>(Obj->GetOuter());
+		if (pkg)
+		{
+			FString AssetName = FPaths::GetBaseFilename(AssetRef.ToString());
+			UObject* AssetObject = FindObject<UObject>(pkg, *AssetName);
+			FAssetEditorManager::Get().OpenEditorForAsset(AssetObject);
+		}
+	}
+}
+
 #undef LOCTEXT_NAMESPACE
diff --git a/Engine/Source/Editor/GameplayAbilitiesEditor/Private/SGameplayCueEditor.h b/Engine/Source/Editor/GameplayAbilitiesEditor/Private/SGameplayCueEditor.h
--- a/Engine/Source/Editor/GameplayAbilitiesEditor/Private/SGameplayCueEditor.h
+++ b/Engine/Source/Editor/GameplayAbilitiesEditor/Private/SGameplayCueEditor.h
@@ -20,4 +20,7 @@ public:
 	virtual FReply OnNewGameplayCueButtonPressed() = 0;
 	
 	static FString GetPathNameForGameplayCueTag(FString Tag);
+
+	static void CreateNewGameplayCueNotifyDialogue(FString GameplayCue);
+	static void OpenEditorForNotify(FString NotifyFullPath);
 };
diff --git a/Engine/Source/Editor/GameplayTagsEditor/GameplayTagsEditor.Build.cs b/Engine/Source/Editor/GameplayTagsEditor/GameplayTagsEditor.Build.cs
--- a/Engine/Source/Editor/GameplayTagsEditor/GameplayTagsEditor.Build.cs
+++ b/Engine/Source/Editor/GameplayTagsEditor/GameplayTagsEditor.Build.cs
@@ -41,6 +41,7 @@ namespace UnrealBuildTool.Rules
 					"ContentBrowser",
 					"MainFrame",
 					"UnrealEd",
+                    "SourceControl",
 				}
 			);
 
diff --git a/Engine/Source/Editor/GameplayTagsEditor/Private/GameplayTagContainerCustomization.cpp b/Engine/Source/Editor/GameplayTagsEditor/Private/GameplayTagContainerCustomization.cpp
--- a/Engine/Source/Editor/GameplayTagsEditor/Private/GameplayTagContainerCustomization.cpp
+++ b/Engine/Source/Editor/GameplayTagsEditor/Private/GameplayTagContainerCustomization.cpp
@@ -146,7 +146,7 @@ FReply FGameplayTagContainerCustomization::OnEditButtonClicked()
 	.Title(Title)
 	.ClientSize(FVector2D(600, 400))
 	[
-		SNew(SGameplayTagWidget, EditableContainers)
+		SAssignNew(GameplayTagWidget, SGameplayTagWidget, EditableContainers)
 		.Filter( Categories )
 		.OnTagChanged(this, &FGameplayTagContainerCustomization::RefreshTagList)
 		.ReadOnly(bReadOnly)
@@ -248,7 +248,10 @@ void FGameplayTagContainerCustomization::OnGameplayTagWidgetWindowDeactivate()
 {
 	if( GameplayTagWidgetWindow.IsValid() )
 	{
-		GameplayTagWidgetWindow->RequestDestroyWindow();
+		if (GameplayTagWidget.IsValid() && GameplayTagWidget->IsAddingNewTag() == false)
+		{
+			GameplayTagWidgetWindow->RequestDestroyWindow();
+		}
 	}
 }
 
diff --git a/Engine/Source/Editor/GameplayTagsEditor/Private/GameplayTagContainerCustomization.h b/Engine/Source/Editor/GameplayTagsEditor/Private/GameplayTagContainerCustomization.h
--- a/Engine/Source/Editor/GameplayTagsEditor/Private/GameplayTagContainerCustomization.h
+++ b/Engine/Source/Editor/GameplayTagsEditor/Private/GameplayTagContainerCustomization.h
@@ -70,5 +70,8 @@ private:
 
 	/** The Window for the GameplayTagWidget */
 	TSharedPtr<SWindow> GameplayTagWidgetWindow;
+
+	/** The widget */
+	TSharedPtr<SGameplayTagWidget> GameplayTagWidget;
 };
 
diff --git a/Engine/Source/Editor/GameplayTagsEditor/Private/SGameplayTagGraphPin.cpp b/Engine/Source/Editor/GameplayTagsEditor/Private/SGameplayTagGraphPin.cpp
--- a/Engine/Source/Editor/GameplayTagsEditor/Private/SGameplayTagGraphPin.cpp
+++ b/Engine/Source/Editor/GameplayTagsEditor/Private/SGameplayTagGraphPin.cpp
@@ -89,20 +89,12 @@ TSharedRef<SWidget> SGameplayTagGraphPin::GetListContent()
 		.AutoHeight()
 		.MaxHeight( 400 )
 		[
-			SNew(SScaleBox)
-			.HAlign(EHorizontalAlignment::HAlign_Left)
-			.VAlign(EVerticalAlignment::VAlign_Top)
-			.StretchDirection(EStretchDirection::DownOnly)
-			.Stretch(EStretch::ScaleToFit)
-			.Content()
-			[
-				SNew( SGameplayTagWidget, EditableContainers )
-				.OnTagChanged( this, &SGameplayTagGraphPin::RefreshTagList )
-				.TagContainerName( TEXT("SGameplayTagGraphPin") )
-				.Visibility( this, &SGraphPin::GetDefaultValueVisibility )
-				.MultiSelect(false)
-				.Filter(FilterString)
-			]
+			SNew( SGameplayTagWidget, EditableContainers )
+			.OnTagChanged( this, &SGameplayTagGraphPin::RefreshTagList )
+			.TagContainerName( TEXT("SGameplayTagGraphPin") )
+			.Visibility( this, &SGraphPin::GetDefaultValueVisibility )
+			.MultiSelect(false)
+			.Filter(FilterString)
 		];
 }
 
diff --git a/Engine/Source/Editor/GameplayTagsEditor/Private/SGameplayTagWidget.cpp b/Engine/Source/Editor/GameplayTagsEditor/Private/SGameplayTagWidget.cpp
--- a/Engine/Source/Editor/GameplayTagsEditor/Private/SGameplayTagWidget.cpp
+++ b/Engine/Source/Editor/GameplayTagsEditor/Private/SGameplayTagWidget.cpp
@@ -8,6 +8,13 @@
 #include "Editor/PropertyEditor/Public/PropertyHandle.h"
 #include "SSearchBox.h"
 #include "SScaleBox.h"
+#include "AssetEditorManager.h"
+#include "AssetToolsModule.h"
+#include "SHyperlink.h"
+#include "SNotificationList.h"
+#include "NotificationManager.h"
+#include "SSearchBox.h"
+#include "GameplayTagsSettings.h"
 
 #define LOCTEXT_NAMESPACE "GameplayTagWidget"
 
@@ -23,8 +30,11 @@ void SGameplayTagWidget::Construct(const FArguments& InArgs, const TArray<FEdita
 	TagContainerName = InArgs._TagContainerName;
 	bMultiSelect = InArgs._MultiSelect;
 	PropertyHandle = InArgs._PropertyHandle;
+	bIsAddingNewTag = false;
+	RootFilterString = InArgs._Filter;
 
 	IGameplayTagsModule::Get().GetGameplayTagsManager().GetFilteredGameplayRootTags(InArgs._Filter, TagItems);
+	bool CanAddFromINI = UGameplayTagsManager::ShouldImportTagsFromINI(); // We only support adding new tags to the ini files.
 
 	// Tag the assets as transactional so they can support undo/redo
 	for (int32 AssetIdx = 0; AssetIdx < TagContainers.Num(); ++AssetIdx)
@@ -42,6 +52,32 @@ void SGameplayTagWidget::Construct(const FArguments& InArgs, const TArray<FEdita
 		.BorderImage(FEditorStyle::GetBrush("ToolPanel.GroupBorder"))
 		[
 			SNew(SVerticalBox)
+			+SVerticalBox::Slot()
+			.AutoHeight()
+			.VAlign(VAlign_Top)
+			[
+				SNew(SHorizontalBox)
+				+ SHorizontalBox::Slot()
+				.Padding(2.0f, 2.0f)
+				.AutoWidth()
+				[
+					SAssignNew(NewTagTextBox, SEditableTextBox)
+					.MinDesiredWidth(210.0f)
+					.HintText(LOCTEXT("NewTag", "X.Y.Z"))
+					.OnTextCommitted(this, &SGameplayTagWidget::OnNewGameplayTagCommited)
+					.Visibility(CanAddFromINI ? EVisibility::Visible : EVisibility::Collapsed)
+				]
+				+ SHorizontalBox::Slot()
+				.Padding(2.0f, 2.0f)
+				.AutoWidth()
+				[
+					SNew(SButton)
+					.Text(LOCTEXT("AddNew", "Add New"))
+					.OnClicked(this, &SGameplayTagWidget::OnNewGameplayTagButtonPressed)
+					.Visibility(CanAddFromINI ? EVisibility::Visible : EVisibility::Collapsed)
+				]
+			]
+
 			+SVerticalBox::Slot()
 			.AutoHeight()
 			.VAlign(VAlign_Top)
@@ -56,6 +92,7 @@ void SGameplayTagWidget::Construct(const FArguments& InArgs, const TArray<FEdita
 				]
 				+SHorizontalBox::Slot()
 				.AutoWidth()
+
 				[
 					SNew(SButton)
 					.OnClicked(this, &SGameplayTagWidget::OnCollapseAllClicked)
@@ -74,7 +111,7 @@ void SGameplayTagWidget::Construct(const FArguments& InArgs, const TArray<FEdita
 				.FillWidth(1.f)
 				.Padding(5,1,5,1)
 				[
-					SNew(SSearchBox)
+					SAssignNew(SearchTagBox, SSearchBox)
 					.HintText(LOCTEXT("GameplayTagWidget_SearchBoxHint", "Search Gameplay Tags"))
 					.OnTextChanged( this, &SGameplayTagWidget::OnFilterTextChanged )
 				]
@@ -104,6 +141,54 @@ void SGameplayTagWidget::Construct(const FArguments& InArgs, const TArray<FEdita
 	VerifyAssetTagValidity();
 }
 
+void SGameplayTagWidget::OnNewGameplayTagCommited(const FText& InText, ETextCommit::Type InCommitType)
+{
+	if (InCommitType == ETextCommit::OnEnter)
+	{
+		CreateNewGameplayTag();
+	}
+}
+
+FReply SGameplayTagWidget::OnNewGameplayTagButtonPressed()
+{
+	CreateNewGameplayTag();
+	return FReply::Handled();
+}
+
+void SGameplayTagWidget::CreateNewGameplayTag()
+{
+	// Only support adding tags via ini file
+	if (UGameplayTagsManager::ShouldImportTagsFromINI() == false)
+	{
+		return;
+	}
+
+	FString str = NewTagTextBox->GetText().ToString();
+	if (str.IsEmpty())
+	{
+		return;
+	}
+
+	// set bIsAddingNewTag, this guards against the window closing when it loses focus due to source control checking out a file
+	TGuardValue<bool>	Guard(bIsAddingNewTag, true);
+
+	UGameplayTagsManager::AddNewGameplayTagToINI(str);
+
+	NewTagTextBox->SetText(FText::GetEmpty());
+
+	IGameplayTagsModule::Get().GetGameplayTagsManager().GetFilteredGameplayRootTags(RootFilterString, TagItems);
+	TagTreeWidget->RequestTreeRefresh();
+
+	auto node = IGameplayTagsModule::Get().GetGameplayTagsManager().FindTagNode(FName(*str));
+	if (node.IsValid())
+	{
+		OnTagChecked(node);
+	}
+
+	// Filter on the new tag
+	SearchTagBox->SetText(FText::FromString(str));
+}
+
 void SGameplayTagWidget::OnFilterTextChanged( const FText& InFilterText )
 {
 	FilterString = InFilterText.ToString();	
@@ -340,10 +425,13 @@ ECheckBoxState SGameplayTagWidget::IsTagChecked(TSharedPtr<FGameplayTagNode> Nod
 			if (Container)
 			{
 				NumValidAssets++;
-				FGameplayTag Tag = TagsManager.RequestGameplayTag(Node->GetCompleteTag());
-				if (Container->HasTag(Tag, EGameplayTagMatchType::Explicit, EGameplayTagMatchType::Explicit))
+				FGameplayTag Tag = TagsManager.RequestGameplayTag(Node->GetCompleteTag(), false);
+				if (Tag.IsValid())
 				{
-					++NumAssetsTagIsAppliedTo;
+					if (Container->HasTag(Tag, EGameplayTagMatchType::Explicit, EGameplayTagMatchType::Explicit))
+					{
+						++NumAssetsTagIsAppliedTo;
+					}
 				}
 			}
 		}
diff --git a/Engine/Source/Editor/GameplayTagsEditor/Private/SGameplayTagWidget.h b/Engine/Source/Editor/GameplayTagsEditor/Private/SGameplayTagWidget.h
--- a/Engine/Source/Editor/GameplayTagsEditor/Private/SGameplayTagWidget.h
+++ b/Engine/Source/Editor/GameplayTagsEditor/Private/SGameplayTagWidget.h
@@ -51,7 +51,12 @@ public:
 	void OnFilterTextChanged( const FText& InFilterText );
 
 	/** Returns true if this TagNode has any children that match the current filter */
-	bool FilterChildrenCheck( TSharedPtr<FGameplayTagNode>  );	
+	bool FilterChildrenCheck( TSharedPtr<FGameplayTagNode>  );
+
+	bool IsAddingNewTag() const
+	{
+		return bIsAddingNewTag;
+	}
 
 private:
 
@@ -64,12 +69,18 @@ private:
 	/* Filter string used during search box */
 	FString FilterString;
 
+	/** root filter (passed in on creation) */
+	FString RootFilterString;
+
 	/* Flag to set if the list is read only*/
 	bool bReadOnly;
 
 	/* Flag to set if we can select multiple items form the list*/
 	bool bMultiSelect;
 
+	/* Flag set while we are in process of adding new tag */
+	bool bIsAddingNewTag;
+
 	/* Array of tags to be displayed in the TreeView*/
 	TArray< TSharedPtr<FGameplayTagNode> > TagItems;
 
@@ -79,6 +90,10 @@ private:
 	/** Tree widget showing the gameplay tag library */
 	TSharedPtr< STreeView< TSharedPtr<FGameplayTagNode> > > TagTreeWidget;
 
+	TSharedPtr<SEditableTextBox> NewTagTextBox;
+
+	TSharedPtr<SSearchBox> SearchTagBox;
+
 	/** Containers to modify */
 	TArray<FEditableGameplayTagContainerDatum> TagContainers;
 
@@ -87,6 +102,12 @@ private:
 
 	TSharedPtr<IPropertyHandle> PropertyHandle;
 
+	void OnNewGameplayTagCommited(const FText& InText, ETextCommit::Type InCommitType);
+
+	FReply OnNewGameplayTagButtonPressed();
+
+	void CreateNewGameplayTag();
+
 	/**
 	 * Generate a row widget for the specified item node and table
 	 * 
