Hide "Add New" button in gameplay cue editor when a tag already has a notify, since adding another notify will not be used in game.

--------
Integrated using branch UE4-To-UE4-Orion (reversed) of change#2669936 by [REDACTED_PERSON] on 2015/08/26 16:50:31.

diff --git a/Engine/Source/Editor/GameplayAbilitiesEditor/Private/SGameplayCueEditor.cpp b/Engine/Source/Editor/GameplayAbilitiesEditor/Private/SGameplayCueEditor.cpp
--- a/Engine/Source/Editor/GameplayAbilitiesEditor/Private/SGameplayCueEditor.cpp
+++ b/Engine/Source/Editor/GameplayAbilitiesEditor/Private/SGameplayCueEditor.cpp
@@ -517,11 +517,22 @@ private:
 				}
 				else if (ColumnName == CueHandlerColumnName)
 				{
+					bool HasNotifies = false;
+					for (auto& HandlerItem : Item->HandlerItems)
+					{
+						if (HandlerItem->GameplayCueNotifyObj.IsValid())
+						{
+							HasNotifies = true;
+							break;
+						}
+					}
+
 					return
 					SNew(SBox)
 					.Padding( FMargin(2.0f, 0.0f) )
 					.HAlign(HAlign_Left)
 					.VAlign(VAlign_Center)
+					.Visibility(HasNotifies ? EVisibility::Collapsed : EVisibility::Visible)
 					[
 						SNew(SButton)
 						.Text(LOCTEXT("AddNew", "Add New"))
