Gameplay Abilities: Added the ability to flag a pin as requiring a connection on a latent ability call (meta=(RequiresConnection=true) on the delegate property)
[INTEGRATE] Change 2711735 by [REDACTED_PERSON][REDACTED_HANDLE] on 2015/09/30 15:56:26

diff --git a/Engine/Source/Editor/GameplayAbilitiesEditor/Classes/K2Node_LatentAbilityCall.h b/Engine/Source/Editor/GameplayAbilitiesEditor/Classes/K2Node_LatentAbilityCall.h
--- a/Engine/Source/Editor/GameplayAbilitiesEditor/Classes/K2Node_LatentAbilityCall.h
+++ b/Engine/Source/Editor/GameplayAbilitiesEditor/Classes/K2Node_LatentAbilityCall.h
@@ -16,6 +16,7 @@ public:
 	// UEdGraphNode interface
 	virtual void GetMenuActions(FBlueprintActionDatabaseRegistrar& ActionRegistrar) const override;
 	virtual bool IsCompatibleWithGraph(UEdGraph const* TargetGraph) const override;
+	virtual void ValidateNodeDuringCompilation(class FCompilerResultsLog& MessageLog) const override;
 	// End of UEdGraphNode interface
 	
 protected:
diff --git a/Engine/Source/Editor/GameplayAbilitiesEditor/Private/K2Node_LatentAbilityCall.cpp b/Engine/Source/Editor/GameplayAbilitiesEditor/Private/K2Node_LatentAbilityCall.cpp
--- a/Engine/Source/Editor/GameplayAbilitiesEditor/Private/K2Node_LatentAbilityCall.cpp
+++ b/Engine/Source/Editor/GameplayAbilitiesEditor/Private/K2Node_LatentAbilityCall.cpp
@@ -1,6 +1,5 @@
 // Copyright 1998-2015 Epic Games, Inc. All Rights Reserved.
 
-
 #include "AbilitySystemEditorPrivatePCH.h"
 #include "Kismet/KismetMathLibrary.h"
 #include "Kismet/KismetArrayLibrary.h"
@@ -13,6 +12,12 @@
 #include "K2Node_EnumLiteral.h"
 #include "BlueprintFunctionNodeSpawner.h"
 #include "BlueprintActionDatabaseRegistrar.h"
+#include "CompilerResultsLog.h"
+
+static FName FK2Node_LatentAbilityCallHelper_RequiresConnection(TEXT("RequiresConnection"));
+
+/////////////////////////////////////////////////////
+// UK2Node_LatentAbilityCall
 
 #define LOCTEXT_NAMESPACE "K2Node"
 
@@ -85,4 +90,28 @@ void UK2Node_LatentAbilityCall::GetMenuActions(FBlueprintActionDatabaseRegistrar
 	}) );
 }
 
+void UK2Node_LatentAbilityCall::ValidateNodeDuringCompilation(class FCompilerResultsLog& MessageLog) const
+{
+	Super::ValidateNodeDuringCompilation(MessageLog);
+
+	UFunction* DelegateSignatureFunction = NULL;
+	for (TFieldIterator<UProperty> PropertyIt(ProxyClass); PropertyIt; ++PropertyIt)
+	{
+		if (UMulticastDelegateProperty* Property = Cast<UMulticastDelegateProperty>(*PropertyIt))
+		{
+			if (Property->GetBoolMetaData(FK2Node_LatentAbilityCallHelper_RequiresConnection))
+			{
+				if (UEdGraphPin* DelegateExecPin = FindPin(Property->GetName()))
+				{
+					if (DelegateExecPin->LinkedTo.Num() < 1)
+					{
+						const FText MessageText = FText::Format(LOCTEXT("NoConnectionToRequiredExecPin", "@@ - Unhandled event.  You need something connected to the '{0}' pin"), FText::FromName(Property->GetFName()));
+						MessageLog.Warning(*MessageText.ToString(), this);
+					}
+				}
+			}
+		}
+	}
+}
+
 #undef LOCTEXT_NAMESPACE
